Fix deprecation warnings after newest sync

#rnx

diff --git a/Engine/Source/Developer/TraceAnalysis/Private/ControlClient.cpp b/Engine/Source/Developer/TraceAnalysis/Private/ControlClient.cpp
--- a/Engine/Source/Developer/TraceAnalysis/Private/ControlClient.cpp
+++ b/Engine/Source/Developer/TraceAnalysis/Private/ControlClient.cpp
@@ -5,6 +5,7 @@
 #include "Misc/CString.h"
 #include "Sockets.h"
 #include "SocketSubsystem.h"
+#include "AddressInfoTypes.h"
 
 namespace Trace
 {
@@ -23,17 +24,22 @@ bool FControlClient::Connect(const TCHAR* Host, uint16 Port)
         return false;
     }
 
-	ISocketSubsystem& Sockets = *(ISocketSubsystem::Get());
-	TSharedPtr<FInternetAddr> Addr = Sockets.CreateInternetAddr();
-	bool bIsHostValid = false;
-	Addr->SetIp(Host, bIsHostValid);
-	if (!bIsHostValid)
+	ISocketSubsystem* Sockets = ISocketSubsystem::Get();
+	if (Sockets == nullptr)
 	{
-		ESocketErrors SocketErrors = Sockets.GetHostByName(TCHAR_TO_ANSI(Host), *Addr);
-		if (SocketErrors != SE_NO_ERROR)
+		return false;
+	}
+	
+	TSharedPtr<FInternetAddr> Addr = Sockets->GetAddressFromString(Host);
+	if (!Addr.IsValid() || !Addr->IsValid())
+	{
+		FAddressInfoResult GAIRequest = Sockets->GetAddressInfo(Host, nullptr, EAddressInfoFlags::Default, NAME_None);
+		if (GAIRequest.ReturnCode != SE_NO_ERROR || GAIRequest.Results.Num() == 0)
 		{
 			return false;
 		}
+		
+		Addr = GAIRequest.Results[0].Address;
 	}
 	Addr->SetPort(Port);
 	return Connect(*Addr);
@@ -46,10 +52,14 @@ bool FControlClient::Connect(const FInternetAddr& Address)
 	{
 		return false;
 	}
+	
+	ISocketSubsystem* Sockets = ISocketSubsystem::Get();	
+	if (Sockets == nullptr)
+	{
+		return false;
+	}
 
-	ISocketSubsystem& Sockets = *(ISocketSubsystem::Get());
-
-	FSocket* ClientSocket = Sockets.CreateSocket(NAME_Stream, TEXT("TraceControlClient"));
+	FSocket* ClientSocket = Sockets->CreateSocket(NAME_Stream, TEXT("TraceControlClient"), Address.GetProtocolType());
 	if (ClientSocket == nullptr)
 	{
 		return false;
@@ -57,7 +67,7 @@ bool FControlClient::Connect(const FInternetAddr& Address)
 
 	if (!ClientSocket->Connect(Address))
 	{
-		Sockets.DestroySocket(ClientSocket);
+		Sockets->DestroySocket(ClientSocket);
 		return false;
 	}
 
diff --git a/Engine/Source/Developer/TraceServices/Private/SessionService.cpp b/Engine/Source/Developer/TraceServices/Private/SessionService.cpp
--- a/Engine/Source/Developer/TraceServices/Private/SessionService.cpp
+++ b/Engine/Source/Developer/TraceServices/Private/SessionService.cpp
@@ -8,6 +8,7 @@
 #include "SocketSubsystem.h"
 #include "Trace/ControlClient.h"
 #include "IPAddress.h"
+#include "AddressInfoTypes.h"
 
 #if PLATFORM_WINDOWS
 	#include "Windows/AllowWindowsPlatformTypes.h"
@@ -185,6 +186,7 @@ bool FSessionService::ConnectSession(const TCHAR* ControlClientAddress)
 	{
 		return false;
 	}
+
 	uint16 Port = 1985;
 	FString AddressString = ControlClientAddress;
 	const int32 LastColonIndex = AddressString.Find(":", ESearchCase::IgnoreCase, ESearchDir::FromEnd);
@@ -194,12 +196,19 @@ bool FSessionService::ConnectSession(const TCHAR* ControlClientAddress)
 		Port = FCString::Atoi(*PortString);
 		AddressString = AddressString.Left(LastColonIndex);
 	}
-	TSharedPtr<FInternetAddr> ClientAddr = Sockets->CreateInternetAddr();
-	ESocketErrors SocketErrors = Sockets->GetHostByName(TCHAR_TO_ANSI(*AddressString), *ClientAddr);
-	if (SocketErrors != SE_NO_ERROR)
+	
+	TSharedPtr<FInternetAddr> ClientAddr = Sockets->GetAddressFromString(AddressString);
+	if (!ClientAddr.IsValid() || !ClientAddr->IsValid())
 	{
-		return false;
+		FAddressInfoResult GAIRequest = Sockets->GetAddressInfo(*AddressString, nullptr, EAddressInfoFlags::Default, NAME_None);
+		if (GAIRequest.ReturnCode != SE_NO_ERROR || GAIRequest.Results.Num() == 0)
+		{
+			return false;
+		}
+		
+		ClientAddr = GAIRequest.Results[0].Address;
 	}
+	
 	ClientAddr->SetPort(Port);
 	FControlClient ControlClient;
 	if (!ControlClient.Connect(*ClientAddr))
