Fixed trace's recorder (hosted in UnrealInsights) failing to bind a socket to a port.

(ushell-p4-cherrypick of 6908891 by [REDACTED_PERSON])

diff --git a/Engine/Source/Developer/TraceAnalysis/Private/Recorder/Recorder.cpp b/Engine/Source/Developer/TraceAnalysis/Private/Recorder/Recorder.cpp
--- a/Engine/Source/Developer/TraceAnalysis/Private/Recorder/Recorder.cpp
+++ b/Engine/Source/Developer/TraceAnalysis/Private/Recorder/Recorder.cpp
@@ -258,6 +258,7 @@ bool FRecorder::StartRecording()
 	// Bind it to the trace-recording port
 	bool bBound = false;
 	TSharedPtr<FInternetAddr> Addr = Sockets.CreateInternetAddr();
+	Addr->SetIp(0);
 	Addr->SetPort(1980);
 	if (Socket->Bind(*Addr))
 	{
