Fixed error in include declarations

diff --git a/Engine/Source/Developer/TraceAnalysis/Private/Store/FileStore.cpp b/Engine/Source/Developer/TraceAnalysis/Private/Store/FileStore.cpp
--- a/Engine/Source/Developer/TraceAnalysis/Private/Store/FileStore.cpp
+++ b/Engine/Source/Developer/TraceAnalysis/Private/Store/FileStore.cpp
@@ -3,8 +3,7 @@
 //
 
 #include "Trace/Store.h"
-#include "Trace/DataStream.h"
-#include "Trace/DataStream.h"
+#include "DataStream.h"
 #include "HAL/PlatformFile.h"
 #include "HAL/PlatformProcess.h"
 #include "Misc/ScopeLock.h"
