Removed incorrect include statement.

diff --git a/Engine/Source/Developer/TraceAnalysis/Public/Trace/Store.h b/Engine/Source/Developer/TraceAnalysis/Public/Trace/Store.h
--- a/Engine/Source/Developer/TraceAnalysis/Public/Trace/Store.h
+++ b/Engine/Source/Developer/TraceAnalysis/Public/Trace/Store.h
@@ -7,7 +7,6 @@
 #include "CoreTypes.h"
 #include "Containers/UnrealString.h"
 #include "Templates/SharedPointer.h"
-#include "Trace/Store.h"
 
 namespace Trace
 {
