Addressed copyright notice false positives.

diff --git a/Engine/Source/Developer/TraceAnalysis/Private/Analysis/Context.cpp b/Engine/Source/Developer/TraceAnalysis/Private/Analysis/Context.cpp
--- a/Engine/Source/Developer/TraceAnalysis/Private/Analysis/Context.cpp
+++ b/Engine/Source/Developer/TraceAnalysis/Private/Analysis/Context.cpp
@@ -1,6 +1,4 @@
-//
 // Copyright 1998-2019 Epic Games, Inc. All Rights Reserved.
-//
 
 #include "Engine.h"
 #include "Trace/Analysis.h"
diff --git a/Engine/Source/Developer/TraceAnalysis/Private/Analysis/Engine.cpp b/Engine/Source/Developer/TraceAnalysis/Private/Analysis/Engine.cpp
--- a/Engine/Source/Developer/TraceAnalysis/Private/Analysis/Engine.cpp
+++ b/Engine/Source/Developer/TraceAnalysis/Private/Analysis/Engine.cpp
@@ -1,6 +1,4 @@
-//
 // Copyright 1998-2019 Epic Games, Inc. All Rights Reserved.
-//
 
 #include "Engine.h"
 #include "Containers/ArrayView.h"
diff --git a/Engine/Source/Developer/TraceAnalysis/Private/Analysis/Engine.h b/Engine/Source/Developer/TraceAnalysis/Private/Analysis/Engine.h
--- a/Engine/Source/Developer/TraceAnalysis/Private/Analysis/Engine.h
+++ b/Engine/Source/Developer/TraceAnalysis/Private/Analysis/Engine.h
@@ -1,6 +1,4 @@
-//
 // Copyright 1998-2019 Epic Games, Inc. All Rights Reserved.
-//
 
 #pragma once
 
diff --git a/Engine/Source/Developer/TraceAnalysis/Private/Analysis/Processor.cpp b/Engine/Source/Developer/TraceAnalysis/Private/Analysis/Processor.cpp
--- a/Engine/Source/Developer/TraceAnalysis/Private/Analysis/Processor.cpp
+++ b/Engine/Source/Developer/TraceAnalysis/Private/Analysis/Processor.cpp
@@ -1,6 +1,4 @@
-//
 // Copyright 1998-2019 Epic Games, Inc. All Rights Reserved.
-//
 
 #include "Engine.h"
 #include "Trace/Analysis.h"
diff --git a/Engine/Source/Developer/TraceAnalysis/Private/Common.h b/Engine/Source/Developer/TraceAnalysis/Private/Common.h
--- a/Engine/Source/Developer/TraceAnalysis/Private/Common.h
+++ b/Engine/Source/Developer/TraceAnalysis/Private/Common.h
@@ -1,6 +1,4 @@
-//
 // Copyright 1998-2019 Epic Games, Inc. All Rights Reserved.
-//
 
 #pragma once
 
diff --git a/Engine/Source/Developer/TraceAnalysis/Private/DataStream.cpp b/Engine/Source/Developer/TraceAnalysis/Private/DataStream.cpp
--- a/Engine/Source/Developer/TraceAnalysis/Private/DataStream.cpp
+++ b/Engine/Source/Developer/TraceAnalysis/Private/DataStream.cpp
@@ -1,6 +1,4 @@
-//
 // Copyright 1998-2019 Epic Games, Inc. All Rights Reserved.
-//
 
 #include "DataStream.h"
 #include "GenericPlatform/GenericPlatformFile.h"
@@ -70,4 +68,4 @@ IInDataStream* DataStream_ReadFile(const TCHAR* FilePath)
 	return new FFileStream(FilePath);
 }
 
-}
\ No newline at end of file
+}
diff --git a/Engine/Source/Developer/TraceAnalysis/Private/DataStream.h b/Engine/Source/Developer/TraceAnalysis/Private/DataStream.h
--- a/Engine/Source/Developer/TraceAnalysis/Private/DataStream.h
+++ b/Engine/Source/Developer/TraceAnalysis/Private/DataStream.h
@@ -1,6 +1,4 @@
-//
 // Copyright 1998-2019 Epic Games, Inc. All Rights Reserved.
-//
 
 #pragma once
 
diff --git a/Engine/Source/Developer/TraceAnalysis/Private/Recorder/Recorder.cpp b/Engine/Source/Developer/TraceAnalysis/Private/Recorder/Recorder.cpp
--- a/Engine/Source/Developer/TraceAnalysis/Private/Recorder/Recorder.cpp
+++ b/Engine/Source/Developer/TraceAnalysis/Private/Recorder/Recorder.cpp
@@ -1,6 +1,4 @@
-//
 // Copyright 1998-2019 Epic Games, Inc. All Rights Reserved.
-//
 
 #include "Trace/Recorder.h"
 #include "CoreTypes.h"
diff --git a/Engine/Source/Developer/TraceAnalysis/Private/Store/FileStore.cpp b/Engine/Source/Developer/TraceAnalysis/Private/Store/FileStore.cpp
--- a/Engine/Source/Developer/TraceAnalysis/Private/Store/FileStore.cpp
+++ b/Engine/Source/Developer/TraceAnalysis/Private/Store/FileStore.cpp
@@ -1,6 +1,4 @@
-//
 // Copyright 1998-2019 Epic Games, Inc. All Rights Reserved.
-//
 
 #include "Trace/Store.h"
 #include "DataStream.h"
diff --git a/Engine/Source/Developer/TraceAnalysis/Public/Trace/Analysis.h b/Engine/Source/Developer/TraceAnalysis/Public/Trace/Analysis.h
--- a/Engine/Source/Developer/TraceAnalysis/Public/Trace/Analysis.h
+++ b/Engine/Source/Developer/TraceAnalysis/Public/Trace/Analysis.h
@@ -1,6 +1,4 @@
-//
 // Copyright 1998-2019 Epic Games, Inc. All Rights Reserved.
-//
 
 #pragma once
 
diff --git a/Engine/Source/Developer/TraceAnalysis/Public/Trace/Analyzer.h b/Engine/Source/Developer/TraceAnalysis/Public/Trace/Analyzer.h
--- a/Engine/Source/Developer/TraceAnalysis/Public/Trace/Analyzer.h
+++ b/Engine/Source/Developer/TraceAnalysis/Public/Trace/Analyzer.h
@@ -1,6 +1,4 @@
-//
 // Copyright 1998-2019 Epic Games, Inc. All Rights Reserved.
-//
 
 #pragma once
 
diff --git a/Engine/Source/Developer/TraceAnalysis/Public/Trace/DataStream.h b/Engine/Source/Developer/TraceAnalysis/Public/Trace/DataStream.h
--- a/Engine/Source/Developer/TraceAnalysis/Public/Trace/DataStream.h
+++ b/Engine/Source/Developer/TraceAnalysis/Public/Trace/DataStream.h
@@ -1,6 +1,4 @@
-//
 // Copyright 1998-2019 Epic Games, Inc. All Rights Reserved.
-//
 
 #pragma once
 
diff --git a/Engine/Source/Developer/TraceAnalysis/Public/Trace/Recorder.h b/Engine/Source/Developer/TraceAnalysis/Public/Trace/Recorder.h
--- a/Engine/Source/Developer/TraceAnalysis/Public/Trace/Recorder.h
+++ b/Engine/Source/Developer/TraceAnalysis/Public/Trace/Recorder.h
@@ -1,6 +1,4 @@
-//
 // Copyright 1998-2019 Epic Games, Inc. All Rights Reserved.
-//
 
 #pragma once
 
diff --git a/Engine/Source/Developer/TraceAnalysis/Public/Trace/Store.h b/Engine/Source/Developer/TraceAnalysis/Public/Trace/Store.h
--- a/Engine/Source/Developer/TraceAnalysis/Public/Trace/Store.h
+++ b/Engine/Source/Developer/TraceAnalysis/Public/Trace/Store.h
@@ -1,6 +1,4 @@
-//
 // Copyright 1998-2019 Epic Games, Inc. All Rights Reserved.
-//
 
 #pragma once
 
diff --git a/Engine/Source/Developer/TraceAnalysis/TraceAnalysis.Build.cs b/Engine/Source/Developer/TraceAnalysis/TraceAnalysis.Build.cs
--- a/Engine/Source/Developer/TraceAnalysis/TraceAnalysis.Build.cs
+++ b/Engine/Source/Developer/TraceAnalysis/TraceAnalysis.Build.cs
@@ -1,6 +1,4 @@
-//
 // Copyright 1998-2019 Epic Games, Inc. All Rights Reserved.
-//
 
 using UnrealBuildTool;
 
diff --git a/Engine/Source/Runtime/TraceLog/Private/Trace/Event.cpp b/Engine/Source/Runtime/TraceLog/Private/Trace/Event.cpp
--- a/Engine/Source/Runtime/TraceLog/Private/Trace/Event.cpp
+++ b/Engine/Source/Runtime/TraceLog/Private/Trace/Event.cpp
@@ -1,6 +1,4 @@
-//
 // Copyright 1998-2019 Epic Games, Inc. All Rights Reserved.
-//
 
 #include "Trace/Private/Event.h"
 
diff --git a/Engine/Source/Runtime/TraceLog/Private/Trace/Trace.cpp b/Engine/Source/Runtime/TraceLog/Private/Trace/Trace.cpp
--- a/Engine/Source/Runtime/TraceLog/Private/Trace/Trace.cpp
+++ b/Engine/Source/Runtime/TraceLog/Private/Trace/Trace.cpp
@@ -1,6 +1,4 @@
-//
 // Copyright 1998-2019 Epic Games, Inc. All Rights Reserved.
-//
 
 #include "Trace/Trace.h"
 #include "Trace/Private/Event.h"
diff --git a/Engine/Source/Runtime/TraceLog/Private/Trace/Windows/Detail.cpp b/Engine/Source/Runtime/TraceLog/Private/Trace/Windows/Detail.cpp
--- a/Engine/Source/Runtime/TraceLog/Private/Trace/Windows/Detail.cpp
+++ b/Engine/Source/Runtime/TraceLog/Private/Trace/Windows/Detail.cpp
@@ -1,6 +1,4 @@
-//
 // Copyright 1998-2019 Epic Games, Inc. All Rights Reserved.
-//
 
 #include "Trace/Trace.h"
 
diff --git a/Engine/Source/Runtime/TraceLog/Private/Trace/Writer.cpp b/Engine/Source/Runtime/TraceLog/Private/Trace/Writer.cpp
--- a/Engine/Source/Runtime/TraceLog/Private/Trace/Writer.cpp
+++ b/Engine/Source/Runtime/TraceLog/Private/Trace/Writer.cpp
@@ -1,6 +1,4 @@
-//
 // Copyright 1998-2019 Epic Games, Inc. All Rights Reserved.
-//
 
 #include "Trace/Private/Trace.h"
 
diff --git a/Engine/Source/Runtime/TraceLog/Public/Trace/Private/Event.h b/Engine/Source/Runtime/TraceLog/Public/Trace/Private/Event.h
--- a/Engine/Source/Runtime/TraceLog/Public/Trace/Private/Event.h
+++ b/Engine/Source/Runtime/TraceLog/Public/Trace/Private/Event.h
@@ -1,6 +1,4 @@
-//
 // Copyright 1998-2019 Epic Games, Inc. All Rights Reserved.
-//
 
 #pragma once
 
diff --git a/Engine/Source/Runtime/TraceLog/Public/Trace/Private/Event.inl b/Engine/Source/Runtime/TraceLog/Public/Trace/Private/Event.inl
--- a/Engine/Source/Runtime/TraceLog/Public/Trace/Private/Event.inl
+++ b/Engine/Source/Runtime/TraceLog/Public/Trace/Private/Event.inl
@@ -1,6 +1,4 @@
-//
 // Copyright 1998-2019 Epic Games, Inc. All Rights Reserved.
-//
 
 #pragma once
 
diff --git a/Engine/Source/Runtime/TraceLog/Public/Trace/Private/Field.h b/Engine/Source/Runtime/TraceLog/Public/Trace/Private/Field.h
--- a/Engine/Source/Runtime/TraceLog/Public/Trace/Private/Field.h
+++ b/Engine/Source/Runtime/TraceLog/Public/Trace/Private/Field.h
@@ -1,6 +1,4 @@
-//
 // Copyright 1998-2019 Epic Games, Inc. All Rights Reserved.
-//
 
 #pragma once
 
diff --git a/Engine/Source/Runtime/TraceLog/Public/Trace/Private/Trace.h b/Engine/Source/Runtime/TraceLog/Public/Trace/Private/Trace.h
--- a/Engine/Source/Runtime/TraceLog/Public/Trace/Private/Trace.h
+++ b/Engine/Source/Runtime/TraceLog/Public/Trace/Private/Trace.h
@@ -1,6 +1,4 @@
-//
 // Copyright 1998-2019 Epic Games, Inc. All Rights Reserved.
-//
 
 #pragma once
 
diff --git a/Engine/Source/Runtime/TraceLog/Public/Trace/Private/Writer.inl b/Engine/Source/Runtime/TraceLog/Public/Trace/Private/Writer.inl
--- a/Engine/Source/Runtime/TraceLog/Public/Trace/Private/Writer.inl
+++ b/Engine/Source/Runtime/TraceLog/Public/Trace/Private/Writer.inl
@@ -1,6 +1,4 @@
-//
 // Copyright 1998-2019 Epic Games, Inc. All Rights Reserved.
-//
 
 #pragma once
 
diff --git a/Engine/Source/Runtime/TraceLog/Public/Trace/Trace.h b/Engine/Source/Runtime/TraceLog/Public/Trace/Trace.h
--- a/Engine/Source/Runtime/TraceLog/Public/Trace/Trace.h
+++ b/Engine/Source/Runtime/TraceLog/Public/Trace/Trace.h
@@ -1,6 +1,4 @@
-//
 // Copyright 1998-2019 Epic Games, Inc. All Rights Reserved.
-//
 
 #pragma once
 
