Rewrote Recorder to use correct subsystems and give user more control.

diff --git a/Engine/Source/Developer/TraceAnalysis/Private/Recorder/Recorder.cpp b/Engine/Source/Developer/TraceAnalysis/Private/Recorder/Recorder.cpp
--- a/Engine/Source/Developer/TraceAnalysis/Private/Recorder/Recorder.cpp
+++ b/Engine/Source/Developer/TraceAnalysis/Private/Recorder/Recorder.cpp
@@ -1,128 +1,247 @@
 // Copyright 1998-2019 Epic Games, Inc. All Rights Reserved.
 
 #include "Trace/Recorder.h"
-#include "CoreTypes.h"
 
-#if PLATFORM_WINDOWS
+#include "Containers/Array.h"
+#include "HAL/Runnable.h"
+#include "HAL/RunnableThread.h"
+#include "IPAddress.h"
+#include "Sockets.h"
+#include "SocketSubsystem.h"
+#include "Trace/DataStream.h"
 #include "Trace/Store.h"
-#include "DataStream.h"
-#include "Templates/UniquePtr.h"
-
-#include "Windows/AllowWindowsPlatformTypes.h"
-#define _WINSOCK_DEPRECATED_NO_WARNINGS  
-#include <winsock2.h>
-#include <ws2tcpip.h>
-#include "Windows/HideWindowsPlatformTypes.h"
-#pragma comment(lib, "ws2_32.lib")
-
-#include <thread>
 
 namespace Trace
 {
 
-class FSocketStream
-	: public IInDataStream
+////////////////////////////////////////////////////////////////////////////////
+class FRecorder
+	: public IRecorder
+	, public FRunnable
 {
 public:
-					FSocketStream(SOCKET InSocket);
-	virtual int32	Read(void* Data, uint32 Size) override;
+						FRecorder(TSharedRef<IStore> InStore);
+	virtual				~FRecorder();
 
 private:
-	SOCKET			Socket;
+	struct FSession
+	{
+		FSocket*		Socket;
+		IOutDataStream*	StoreSession;
+		bool			bDead;
+	};
+
+	virtual uint32		Run() override;
+	virtual bool		IsRunning() const override;
+	virtual bool		StartRecording() override;
+	virtual void		StopRecording() override;
+	virtual uint32		GetSessionCount() const override;
+	TSharedRef<IStore>	Store;
+	TArray<FSession>	Sessions;
+	FRunnableThread*	Thread = nullptr;
+	FSocket*			ListenSocket = nullptr;
+	volatile bool		bStopRequested;
 };
 
-FSocketStream::FSocketStream(SOCKET InSocket)
-: Socket(InSocket)
+////////////////////////////////////////////////////////////////////////////////
+FRecorder::FRecorder(TSharedRef<IStore> InStore)
+: Store(InStore)
 {
 }
 
-int32 FSocketStream::Read(void* Data, uint32 Size)
+////////////////////////////////////////////////////////////////////////////////
+FRecorder::~FRecorder()
 {
-	return recv(Socket, (char*)Data, Size, 0);
+	StopRecording();
 }
 
-static void Recorder_Record(TSharedRef<IStore> TraceStore, IInDataStream& DataStream)
+////////////////////////////////////////////////////////////////////////////////
+uint32 FRecorder::Run()
 {
-	static const uint32 BufferSize = 1 << 20;
-	void* Buffer = FPlatformMemory::MemoryRangeReserve(BufferSize, true /*Commit*/);
+	ISocketSubsystem& Sockets = *(ISocketSubsystem::Get());
 
-	TUniquePtr<IOutDataStream> File = TUniquePtr<IOutDataStream>(TraceStore->CreateNewSession());
-	
-	while (true)
+	auto KillSession = [&] (FSession& Session)
 	{
-		int32 ReadSize = DataStream.Read(Buffer, BufferSize);
-		if (ReadSize <= 0)
+		Session.Socket->Close();
+		Sockets.DestroySocket(Session.Socket);
+		delete Session.StoreSession;
+	};
+
+	static const uint32 BufferSize = 5 * 1024 * 1024;
+	uint8* Buffer = new uint8[BufferSize];
+
+	while (!bStopRequested)
+	{
+		bool bPending;
+		if (!ListenSocket->HasPendingConnection(bPending))
 		{
 			break;
 		}
 
-		File->Write(Buffer, ReadSize);
-	}
-
-	FPlatformMemory::MemoryRangeFree(Buffer, BufferSize);
-}
+		if (bPending)
+		{
+			FSocket* Socket = ListenSocket->Accept(TEXT("TraceRecClient"));
 
-////////////////////////////////////////////////////////////////////////////////
-static void Recorder_SocketServer(TSharedRef<IStore> TraceStore)
-{
-	WSADATA WsaData;
-	int Result = WSAStartup(MAKEWORD(2, 2), &WsaData);
-	check(Result == 0);
+			bool bAcceptable = false;
+			if (Socket->Wait(ESocketWaitConditions::WaitForRead, FTimespan(ETimespan::TicksPerSecond / 2)))
+			{
+				uint32 Magic;
+				int32 RecvSize;
+				if (Socket->Recv((uint8*)&Magic, sizeof(Magic), RecvSize))
+				{
+					bAcceptable = (RecvSize == sizeof(Magic)) & (Magic == 'TRCE');
+				}
+			}
 
-	SOCKET ListenSocket = socket(AF_INET, SOCK_STREAM, IPPROTO_TCP);
-	check(ListenSocket != INVALID_SOCKET);
+			if (bAcceptable)
+			{
+				IOutDataStream* StoreSession = Store->CreateNewSession();
+				FSession Session = { Socket, StoreSession, false };
+				Sessions.Add(Session);
+			}
+			else
+			{
+				Sockets.DestroySocket(Socket);
+			}
+		}
 
-	sockaddr_in ServerDesc;
-	ServerDesc.sin_family = AF_INET;
-	ServerDesc.sin_addr.s_addr = {};
-	ServerDesc.sin_port = htons(1980);
-	Result = bind(ListenSocket, (SOCKADDR*)&ServerDesc, sizeof(ServerDesc));
-	check(Result != SOCKET_ERROR);
+		if (Sessions.Num() == 0)
+		{
+			FPlatformProcess::SleepNoStats(0.1f);
+			continue;
+		}
 
-	Result = listen(ListenSocket, 1);
-	check(Result != SOCKET_ERROR);
+		for (FSession& Session : Sessions)
+		{
+			uint32 PendingSize;
+			FSocket& Socket = *(Session.Socket);
+			if (!Socket.HasPendingData(PendingSize))
+			{
+				if (Socket.GetConnectionState() != SCS_Connected)
+				{
+					Session.bDead = true;
+				}
+				continue;
+			}
 
-	while (true)
-	{
-		sockaddr_in SockAddr = {};
-		int SockAddrSize = sizeof(SockAddr);
-		SOCKET Socket = accept(ListenSocket, (sockaddr*)&SockAddr, &SockAddrSize);
-		check(Socket != INVALID_SOCKET);
+			while (PendingSize)
+			{
+				int32 RecvSize = (PendingSize > BufferSize) ? BufferSize : PendingSize;
+				if (!Socket.Recv(Buffer, RecvSize, RecvSize))
+				{
+					Session.bDead = true;
+					break;
+				}
+
+				Session.StoreSession->Write(Buffer, RecvSize);
+				PendingSize -= RecvSize;
+			}
+		}
 
-		auto ClientLambda = [] (TSharedRef<IStore> TraceStore, SOCKET Socket)
+		// Reap dead sessions
+		int SessionCount = Sessions.Num();
+		for (int i = 0; i < SessionCount; ++i)
 		{
-			uint32 Magic;
-			int Recvd = recv(Socket, (char*)&Magic, sizeof(Magic), 0);
-			if (Recvd == sizeof(Magic) && Magic == 'TRCE')
+			FSession& Session = Sessions[i];
+			if (!Session.bDead)
 			{
-				FSocketStream Stream(Socket);
-				Recorder_Record(TraceStore, Stream);
+				continue;
 			}
 
-			shutdown(Socket, SD_BOTH);
-			closesocket(Socket);
-		};
+			KillSession(Session);
+
+			--SessionCount;
+			Sessions[i] = Sessions[SessionCount];
+		}
+		Sessions.SetNumUnsafeInternal(SessionCount);
+	}
 
-		std::thread(ClientLambda, TraceStore, Socket);
+	for (FSession& Session : Sessions)
+	{
+		KillSession(Session);
 	}
 
-	closesocket(ListenSocket);
+	ListenSocket->Close();
+	Sockets.DestroySocket(ListenSocket);
+	ListenSocket = nullptr;
+
+	delete[] Buffer;
+	return 0;
 }
 
-void Recorder_StartServer(TSharedRef<IStore> TraceStore)
+////////////////////////////////////////////////////////////////////////////////
+bool FRecorder::IsRunning() const
 {
-	std::thread(Recorder_SocketServer, TraceStore);
+	return (ListenSocket != nullptr);
 }
 
+////////////////////////////////////////////////////////////////////////////////
+bool FRecorder::StartRecording()
+{
+	if (IsRunning())
+	{
+		return true;
+	}
+
+	StopRecording();
+
+	ISocketSubsystem& Sockets = *(ISocketSubsystem::Get());
+
+	// Create a socket to use for listening for inbound connections
+	FSocket* Socket = Sockets.CreateSocket(NAME_Stream, TEXT("TraceRecListen"));
+	if (Socket == nullptr)
+	{
+		return false;
+	}
+
+	// Bind it to the trace-recording port
+	bool bBound = false;
+	TSharedPtr<FInternetAddr> Addr = Sockets.CreateInternetAddr();
+	Addr->SetPort(1980);
+	if (Socket->Bind(*Addr))
+	{
+		bBound = true;
+	}
+
+	// Set the socket listening for connections to accept
+	if (!bBound || !Socket->Listen(32))
+	{
+		Sockets.DestroySocket(Socket);
+		return false;
+	}
+
+	bStopRequested = false;
+	ListenSocket = Socket;
+	Thread = FRunnableThread::Create(this, TEXT("TraceRec"));
+	return true;
 }
 
-#else
+////////////////////////////////////////////////////////////////////////////////
+void FRecorder::StopRecording()
+{
+	if (Thread == nullptr || bStopRequested)
+	{
+		return;
+	}
+
+	bStopRequested = true;
+	Thread->Kill(true);
+	delete Thread;
+	Thread = nullptr;
+}
 
-namespace Trace
+////////////////////////////////////////////////////////////////////////////////
+uint32 FRecorder::GetSessionCount() const
 {
+	return Sessions.Num();
+}
+
 
-void Recorder_StartServer(TSharedRef<IStore> TraceStore) {}
 
+////////////////////////////////////////////////////////////////////////////////
+TSharedPtr<IRecorder> Recorder_Create(TSharedRef<IStore> Store)
+{
+	return MakeShared<FRecorder>(Store);
 }
 
-#endif
+} // namespace Trace
diff --git a/Engine/Source/Developer/TraceAnalysis/Public/Trace/Recorder.h b/Engine/Source/Developer/TraceAnalysis/Public/Trace/Recorder.h
--- a/Engine/Source/Developer/TraceAnalysis/Public/Trace/Recorder.h
+++ b/Engine/Source/Developer/TraceAnalysis/Public/Trace/Recorder.h
@@ -9,6 +9,25 @@ namespace Trace
 
 class IStore;
 
-TRACEANALYSIS_API void Recorder_StartServer(TSharedRef<IStore> TraceStore);
+////////////////////////////////////////////////////////////////////////////////
+class IRecorder
+{
+public:
+	virtual ~IRecorder() = default;
+
+	/** Checks if the recorder is actively recording, return false if it is not. */
+	virtual bool IsRunning() const = 0;
+
+	/** Starts recording, returning true if it was started successfully or if was already running. */
+	virtual bool StartRecording() = 0;
+
+	/** Stops a running recorder dead in its tracks. */
+	virtual void StopRecording() = 0;
+
+	/** Returns the number of current in-bound sessions. */
+	virtual uint32 GetSessionCount() const = 0;
+};
+
+TRACEANALYSIS_API TSharedPtr<IRecorder> Recorder_Create(TSharedRef<IStore> Store);
 
-}
+} // namespace Trace
diff --git a/Engine/Source/Developer/TraceAnalysis/TraceAnalysis.Build.cs b/Engine/Source/Developer/TraceAnalysis/TraceAnalysis.Build.cs
--- a/Engine/Source/Developer/TraceAnalysis/TraceAnalysis.Build.cs
+++ b/Engine/Source/Developer/TraceAnalysis/TraceAnalysis.Build.cs
@@ -7,5 +7,6 @@ public class TraceAnalysis : ModuleRules
 	public TraceAnalysis(ReadOnlyTargetRules Target) : base(Target)
 	{
 		PrivateDependencyModuleNames.Add("Core");
+		PrivateDependencyModuleNames.Add("Sockets");
 	}
 }
diff --git a/Engine/Source/Developer/TraceInsights/Private/Insights/TraceInsightsModule.cpp b/Engine/Source/Developer/TraceInsights/Private/Insights/TraceInsightsModule.cpp
--- a/Engine/Source/Developer/TraceInsights/Private/Insights/TraceInsightsModule.cpp
+++ b/Engine/Source/Developer/TraceInsights/Private/Insights/TraceInsightsModule.cpp
@@ -127,6 +127,8 @@ void FTraceInsightsModule::StartupModule()
 
 void FTraceInsightsModule::ShutdownModule()
 {
+	TraceSessionService->StopRecorderServer();
+
 	FGlobalTabmanager::Get()->UnregisterNomadTabSpawner(StartPageTabName);
 	FGlobalTabmanager::Get()->UnregisterNomadTabSpawner(TimingProfilerTabName);
 	FGlobalTabmanager::Get()->UnregisterNomadTabSpawner(IoProfilerTabName);
diff --git a/Engine/Source/Developer/TraceServices/Private/SessionService.cpp b/Engine/Source/Developer/TraceServices/Private/SessionService.cpp
--- a/Engine/Source/Developer/TraceServices/Private/SessionService.cpp
+++ b/Engine/Source/Developer/TraceServices/Private/SessionService.cpp
@@ -12,7 +12,7 @@ namespace Trace
 FSessionService::FSessionService(TSharedRef<IStore> InTraceStore)
 	: TraceStore(InTraceStore)
 {
-	
+	TraceRecorder = Recorder_Create(TraceStore);
 }
 
 FSessionService::~FSessionService()
@@ -20,9 +20,19 @@ FSessionService::~FSessionService()
 	
 }
 
-void FSessionService::StartRecorderServer()
+bool FSessionService::StartRecorderServer()
+{
+	return TraceRecorder->StartRecording();
+}
+
+bool FSessionService::IsRecorderServerRunning() const
 {
-	Trace::Recorder_StartServer(TraceStore);
+	return (TraceRecorder == nullptr) ? false : TraceRecorder->IsRunning();
+}
+
+void FSessionService::StopRecorderServer()
+{
+	TraceRecorder->StopRecording();
 }
 
 void FSessionService::GetAvailableSessions(TArray<Trace::FSessionHandle>& OutSessions)
@@ -60,4 +70,4 @@ Trace::IInDataStream* FSessionService::OpenSessionFromFile(const TCHAR* FilePath
 	OnSessionSelected().Broadcast(SelectedSession);
 }*/
 
-}
\ No newline at end of file
+}
diff --git a/Engine/Source/Developer/TraceServices/Private/SessionServicePrivate.h b/Engine/Source/Developer/TraceServices/Private/SessionServicePrivate.h
--- a/Engine/Source/Developer/TraceServices/Private/SessionServicePrivate.h
+++ b/Engine/Source/Developer/TraceServices/Private/SessionServicePrivate.h
@@ -7,13 +7,17 @@
 namespace Trace
 {
 
+class IRecorder;
+
 class FSessionService
 	: public ISessionService
 {
 public:
 	FSessionService(TSharedRef<IStore> InTraceStore);
 	virtual ~FSessionService();
-	virtual void StartRecorderServer() override;
+	virtual bool StartRecorderServer() override;
+	virtual bool IsRecorderServerRunning() const override;
+	virtual void StopRecorderServer() override;
 	virtual void GetAvailableSessions(TArray<Trace::FSessionHandle>& OutSessions) override;
 	virtual bool GetSessionInfo(Trace::FSessionHandle SessionHandle, Trace::FSessionInfo& OutSessionInfo) override;
 	virtual Trace::IInDataStream* OpenSessionStream(Trace::FSessionHandle SessionHandle) override;
@@ -21,6 +25,7 @@ public:
 
 private:
 	TSharedRef<IStore> TraceStore;
+	TSharedPtr<IRecorder> TraceRecorder;
 };
 
-}
\ No newline at end of file
+}
diff --git a/Engine/Source/Developer/TraceServices/Private/TraceServicesModule.cpp b/Engine/Source/Developer/TraceServices/Private/TraceServicesModule.cpp
--- a/Engine/Source/Developer/TraceServices/Private/TraceServicesModule.cpp
+++ b/Engine/Source/Developer/TraceServices/Private/TraceServicesModule.cpp
@@ -4,7 +4,6 @@
 #include "TraceServices/ITraceServicesModule.h"
 #include "SessionServicePrivate.h"
 #include "AnalysisServicePrivate.h"
-#include "Trace/Recorder.h"
 
 class FTraceServicesModule
 	: public ITraceServicesModule
diff --git a/Engine/Source/Developer/TraceServices/Public/TraceServices/SessionService.h b/Engine/Source/Developer/TraceServices/Public/TraceServices/SessionService.h
--- a/Engine/Source/Developer/TraceServices/Public/TraceServices/SessionService.h
+++ b/Engine/Source/Developer/TraceServices/Public/TraceServices/SessionService.h
@@ -14,7 +14,9 @@ class ISessionService
 {
 public:
 	virtual ~ISessionService() = default;
-	virtual void StartRecorderServer() = 0;
+	virtual bool StartRecorderServer() = 0;
+	virtual bool IsRecorderServerRunning() const = 0;
+	virtual void StopRecorderServer() = 0;
 	virtual void GetAvailableSessions(TArray<Trace::FSessionHandle>& OutSessions) = 0;
 	virtual bool GetSessionInfo(Trace::FSessionHandle SessionHandle, Trace::FSessionInfo& OutSessionInfo) = 0;
 	virtual Trace::IInDataStream* OpenSessionStream(Trace::FSessionHandle SessionHandle) = 0;
