A select/[recv|accept] pattern is not compatible with FSockets API so we'll give each accepted connection its own thread.

diff --git a/Engine/Source/Developer/TraceAnalysis/Private/Recorder/Recorder.cpp b/Engine/Source/Developer/TraceAnalysis/Private/Recorder/Recorder.cpp
--- a/Engine/Source/Developer/TraceAnalysis/Private/Recorder/Recorder.cpp
+++ b/Engine/Source/Developer/TraceAnalysis/Private/Recorder/Recorder.cpp
@@ -24,25 +24,60 @@ public:
 	virtual				~FRecorder();
 
 private:
-	struct FSession
-	{
-		FSocket*		Socket;
-		IOutDataStream*	StoreSession;
-		bool			bDead;
-	};
-
+	struct				FSession;
 	virtual uint32		Run() override;
 	virtual bool		IsRunning() const override;
 	virtual bool		StartRecording() override;
 	virtual void		StopRecording() override;
 	virtual uint32		GetSessionCount() const override;
+	FSession*			AcceptSession(FSocket& Socket);
+	void				CloseSession(FSession& Session);
+	void				ReapDeadSessions();
 	TSharedRef<IStore>	Store;
-	TArray<FSession>	Sessions;
+	TArray<FSession*>	Sessions;
 	FRunnableThread*	Thread = nullptr;
 	FSocket*			ListenSocket = nullptr;
 	volatile bool		bStopRequested;
 };
 
+
+
+////////////////////////////////////////////////////////////////////////////////
+struct FRecorder::FSession
+	: public FRunnable
+{
+	virtual uint32		Run() override;
+	FSocket*			Socket;
+	IOutDataStream*		StoreSession;
+	FRunnableThread*	Thread;
+	volatile bool		bDead = false;
+};
+
+////////////////////////////////////////////////////////////////////////////////
+uint32 FRecorder::FSession::Run()
+{
+	static const uint32 BufferSize = 1 * 1024 * 1024;
+	uint8* Buffer = new uint8[BufferSize];
+
+	do
+	{
+		int32 RecvSize;
+		if (!Socket->Recv(Buffer, BufferSize, RecvSize))
+		{
+			bDead = true;
+			break;
+		}
+
+		StoreSession->Write(Buffer, RecvSize);
+	}
+	while (!bDead);
+
+	delete[] Buffer;
+	return 0;
+}
+
+
+
 ////////////////////////////////////////////////////////////////////////////////
 FRecorder::FRecorder(TSharedRef<IStore> InStore)
 : Store(InStore)
@@ -56,24 +91,78 @@ FRecorder::~FRecorder()
 }
 
 ////////////////////////////////////////////////////////////////////////////////
-uint32 FRecorder::Run()
+FRecorder::FSession* FRecorder::AcceptSession(FSocket& Socket)
+{
+	bool bAcceptable = false;
+	if (Socket.Wait(ESocketWaitConditions::WaitForRead, FTimespan(ETimespan::TicksPerSecond / 3)))
+	{
+		uint32 Magic;
+		int32 RecvSize;
+		if (Socket.Recv((uint8*)&Magic, sizeof(Magic), RecvSize))
+		{
+			bAcceptable = (RecvSize == sizeof(Magic)) & (Magic == 'TRCE');
+		}
+	}
+
+	if (!bAcceptable)
+	{
+		return nullptr;
+	}
+
+	FSession* Session = new FSession();
+	Session->Socket = &Socket;
+	Session->StoreSession = Store->CreateNewSession();
+	Session->Thread = FRunnableThread::Create(Session, TEXT("TraceRecSession"));
+	return Session;
+}
+
+////////////////////////////////////////////////////////////////////////////////
+void FRecorder::CloseSession(FSession& Session)
 {
+	Session.bDead = true;
+	Session.Socket->Close();
+
+	Session.Thread->Kill(true);
+
 	ISocketSubsystem& Sockets = *(ISocketSubsystem::Get());
+	Sockets.DestroySocket(Session.Socket);
 
-	auto KillSession = [&] (FSession& Session)
+	delete Session.Thread;
+	delete Session.StoreSession;
+	delete &Session;
+}
+
+////////////////////////////////////////////////////////////////////////////////
+void FRecorder::ReapDeadSessions()
+{
+	// Reap dead sessions
+	int SessionCount = Sessions.Num();
+	for (int i = 0; i < SessionCount; ++i)
 	{
-		Session.Socket->Close();
-		Sockets.DestroySocket(Session.Socket);
-		delete Session.StoreSession;
-	};
+		FSession* Session = Sessions[i];
+		if (!Session->bDead)
+		{
+			continue;
+		}
 
-	static const uint32 BufferSize = 5 * 1024 * 1024;
-	uint8* Buffer = new uint8[BufferSize];
+		CloseSession(*Session);
+
+		--SessionCount;
+		Sessions[i] = Sessions[SessionCount];
+	}
+	Sessions.SetNumUnsafeInternal(SessionCount);
+}
+
+////////////////////////////////////////////////////////////////////////////////
+uint32 FRecorder::Run()
+{
+	ISocketSubsystem& Sockets = *(ISocketSubsystem::Get());
 
+	FTimespan WaitTimespan = FTimespan((ETimespan::TicksPerSecond * 3) / 7);
 	while (!bStopRequested)
 	{
 		bool bPending;
-		if (!ListenSocket->HasPendingConnection(bPending))
+		if (!ListenSocket->WaitForPendingConnection(bPending, WaitTimespan))
 		{
 			break;
 		}
@@ -81,22 +170,13 @@ uint32 FRecorder::Run()
 		if (bPending)
 		{
 			FSocket* Socket = ListenSocket->Accept(TEXT("TraceRecClient"));
-
-			bool bAcceptable = false;
-			if (Socket->Wait(ESocketWaitConditions::WaitForRead, FTimespan(ETimespan::TicksPerSecond / 2)))
+			if (Socket == nullptr)
 			{
-				uint32 Magic;
-				int32 RecvSize;
-				if (Socket->Recv((uint8*)&Magic, sizeof(Magic), RecvSize))
-				{
-					bAcceptable = (RecvSize == sizeof(Magic)) & (Magic == 'TRCE');
-				}
+				break;
 			}
 
-			if (bAcceptable)
+			if (FSession* Session = AcceptSession(*Socket))
 			{
-				IOutDataStream* StoreSession = Store->CreateNewSession();
-				FSession Session = { Socket, StoreSession, false };
 				Sessions.Add(Session);
 			}
 			else
@@ -105,67 +185,14 @@ uint32 FRecorder::Run()
 			}
 		}
 
-		if (Sessions.Num() == 0)
-		{
-			FPlatformProcess::SleepNoStats(0.1f);
-			continue;
-		}
-
-		for (FSession& Session : Sessions)
-		{
-			uint32 PendingSize;
-			FSocket& Socket = *(Session.Socket);
-			if (!Socket.HasPendingData(PendingSize))
-			{
-				if (Socket.GetConnectionState() != SCS_Connected)
-				{
-					Session.bDead = true;
-				}
-				continue;
-			}
-
-			while (PendingSize)
-			{
-				int32 RecvSize = (PendingSize > BufferSize) ? BufferSize : PendingSize;
-				if (!Socket.Recv(Buffer, RecvSize, RecvSize))
-				{
-					Session.bDead = true;
-					break;
-				}
-
-				Session.StoreSession->Write(Buffer, RecvSize);
-				PendingSize -= RecvSize;
-			}
-		}
-
-		// Reap dead sessions
-		int SessionCount = Sessions.Num();
-		for (int i = 0; i < SessionCount; ++i)
-		{
-			FSession& Session = Sessions[i];
-			if (!Session.bDead)
-			{
-				continue;
-			}
-
-			KillSession(Session);
-
-			--SessionCount;
-			Sessions[i] = Sessions[SessionCount];
-		}
-		Sessions.SetNumUnsafeInternal(SessionCount);
+		ReapDeadSessions();
 	}
 
-	for (FSession& Session : Sessions)
+	for (FSession* Session : Sessions)
 	{
-		KillSession(Session);
+		CloseSession(*Session);
 	}
 
-	ListenSocket->Close();
-	Sockets.DestroySocket(ListenSocket);
-	ListenSocket = nullptr;
-
-	delete[] Buffer;
 	return 0;
 }
 
@@ -224,10 +251,17 @@ void FRecorder::StopRecording()
 		return;
 	}
 
+	ListenSocket->Close();
+
 	bStopRequested = true;
+
 	Thread->Kill(true);
 	delete Thread;
 	Thread = nullptr;
+
+	ISocketSubsystem& Sockets = *(ISocketSubsystem::Get());
+	Sockets.DestroySocket(ListenSocket);
+	ListenSocket = nullptr;
 }
 
 ////////////////////////////////////////////////////////////////////////////////
