Removed block-scoped static concrete FEventData which was preventing multiple analysis sessions.

diff --git a/Engine/Source/Developer/TraceAnalysis/Private/Analysis/Engine.cpp b/Engine/Source/Developer/TraceAnalysis/Private/Analysis/Engine.cpp
--- a/Engine/Source/Developer/TraceAnalysis/Private/Analysis/Engine.cpp
+++ b/Engine/Source/Developer/TraceAnalysis/Private/Analysis/Engine.cpp
@@ -212,6 +212,77 @@ bool FLz4TransportReader::NextChunk()
 
 
 
+////////////////////////////////////////////////////////////////////////////////
+struct FAnalysisEngine::FEventDataImpl
+	: public FEventData
+{
+	virtual					~FEventDataImpl() = default;
+	virtual const FValue&	GetValue(const ANSICHAR* FieldName) const override;
+	virtual const FArray&	GetArray(const ANSICHAR* FieldName) const override;
+	virtual const uint8*	GetData() const override;
+	virtual const uint8*	GetAttachment() const override;
+	virtual uint16			GetAttachmentSize() const override;
+	virtual uint16			GetTotalSize() const override;
+	const FDispatch*		Dispatch;
+	const uint8*			Ptr;
+	uint16					Size;
+};
+
+////////////////////////////////////////////////////////////////////////////////
+const IAnalyzer::FValue& FAnalysisEngine::FEventDataImpl::GetValue(const ANSICHAR* FieldName) const
+{
+	UPTRINT Ret = ~0ull;
+
+	FFnv1aHash Hash;
+	Hash.Add(FieldName);
+	uint32 NameHash = Hash.Get();
+
+	for (int i = 0, n = Dispatch->FieldCount; i < n; ++i)
+	{
+		const auto& Field = Dispatch->Fields[i];
+		if (Field.Hash == NameHash)
+		{
+			Ret = UPTRINT(Ptr + Field.Offset);
+			Ret |= UPTRINT(Field.TypeInfo) << 48;
+			break;
+		}
+	}
+
+	return *(FValue*)Ret;
+}
+
+////////////////////////////////////////////////////////////////////////////////
+const IAnalyzer::FArray& FAnalysisEngine::FEventDataImpl::GetArray(const ANSICHAR* FieldName) const
+{
+	return *(FArray*)0x493;
+}
+
+////////////////////////////////////////////////////////////////////////////////
+const uint8* FAnalysisEngine::FEventDataImpl::GetData() const
+{
+	return Ptr;
+}
+
+////////////////////////////////////////////////////////////////////////////////
+const uint8* FAnalysisEngine::FEventDataImpl::GetAttachment() const
+{
+	return Ptr + Dispatch->EventSize;
+}
+
+////////////////////////////////////////////////////////////////////////////////
+uint16 FAnalysisEngine::FEventDataImpl::GetAttachmentSize() const
+{
+	return Size - Dispatch->EventSize;
+}
+
+////////////////////////////////////////////////////////////////////////////////
+uint16 FAnalysisEngine::FEventDataImpl::GetTotalSize() const
+{
+	return Size;
+}
+
+
+
 ////////////////////////////////////////////////////////////////////////////////
 enum ERouteId : uint16
 {
@@ -224,6 +295,8 @@ enum ERouteId : uint16
 FAnalysisEngine::FAnalysisEngine(TArray<IAnalyzer*>&& InAnalyzers)
 : Analyzers(MoveTemp(InAnalyzers))
 {
+	EventDataImpl = new FEventDataImpl();
+
 	uint16 SelfIndex = Analyzers.Num();
 	Analyzers.Add(this);
 
@@ -247,6 +320,8 @@ FAnalysisEngine::~FAnalysisEngine()
 	{
 		FMemory::Free(Dispatch);
 	}
+
+	delete EventDataImpl;
 }
 
 ////////////////////////////////////////////////////////////////////////////////
@@ -476,62 +551,6 @@ bool FAnalysisEngine::EstablishTransport(FStreamReader::FData& Data)
 ////////////////////////////////////////////////////////////////////////////////
 bool FAnalysisEngine::OnData(FStreamReader::FData& Data)
 {
-	static struct : public FEventData
-	{
-		const FValue& GetValue(const ANSICHAR* FieldName) const
-		{
-			UPTRINT Ret = ~0ull;
-
-			FFnv1aHash Hash;
-			Hash.Add(FieldName);
-			uint32 NameHash = Hash.Get();
-
-			for (int i = 0, n = Dispatch->FieldCount; i < n; ++i)
-			{
-				const auto& Field = Dispatch->Fields[i];
-				if (Field.Hash == NameHash)
-				{
-					Ret = UPTRINT(Ptr + Field.Offset);
-					Ret |= UPTRINT(Field.TypeInfo) << 48;
-					break;
-				}
-			}
-
-			return *(FValue*)Ret;
-		}
-
-		const FArray& GetArray(const ANSICHAR* FieldName) const override
-		{
-			return *(FArray*)0x493;
-		}
-
-		const uint8* GetData() const override
-		{
-			return Ptr;
-		}
-
-		const uint8* GetAttachment() const override
-		{
-			return Ptr + Dispatch->EventSize;
-		}
-
-		uint16 GetAttachmentSize() const override
-		{
-			return Size - Dispatch->EventSize;
-		}
-
-		uint16 GetTotalSize() const override
-		{
-			return Size;
-		}
-
-		const FDispatch*	Dispatch;
-		const uint8*		Ptr;
-		uint16				Size;
-	} EventData;
-
-
-
 	if (Transport == nullptr)
 	{
 		if (!EstablishTransport(Data))
@@ -573,15 +592,15 @@ bool FAnalysisEngine::OnData(FStreamReader::FData& Data)
 
 		const FDispatch* Dispatch = Dispatches[Header->Uid];
 
-		EventData.Dispatch = Dispatch;
-		EventData.Ptr = Header->EventData;
-		EventData.Size = Header->Size;
+		EventDataImpl->Dispatch = Dispatch;
+		EventDataImpl->Ptr = Header->EventData;
+		EventDataImpl->Size = Header->Size;
 
 		const FRoute* Route = Routes.GetData() + Dispatch->FirstRoute;
 		for (uint32 n = Route->Count; n--; ++Route)
 		{
 			IAnalyzer* Analyzer = Analyzers[Route->AnalyzerIndex];
-			Analyzer->OnEvent(Route->Id, { SessionContext, EventData });
+			Analyzer->OnEvent(Route->Id, { SessionContext, *EventDataImpl });
 		}
 	}
 
diff --git a/Engine/Source/Developer/TraceAnalysis/Private/Analysis/Engine.h b/Engine/Source/Developer/TraceAnalysis/Private/Analysis/Engine.h
--- a/Engine/Source/Developer/TraceAnalysis/Private/Analysis/Engine.h
+++ b/Engine/Source/Developer/TraceAnalysis/Private/Analysis/Engine.h
@@ -53,6 +53,8 @@ private:
 		uint16			AnalyzerIndex;
 	};
 
+	struct				FEventDataImpl;
+
 	virtual void		OnAnalysisBegin(const FOnAnalysisContext& Context) override;
 	virtual void		OnAnalysisEnd() override;
 	virtual void		OnEvent(uint16 RouteId, const FOnEventContext& Context) override;
@@ -70,6 +72,7 @@ private:
 	TArray<IAnalyzer*>	Analyzers;
 	TArray<FDispatch*>	Dispatches;
 	FTransportReader*	Transport = nullptr;
+	FEventDataImpl*		EventDataImpl;
 };
 
 } // namespace Trace
