Found a couple of "bool b[A-Z]..." convention divergences.

diff --git a/Engine/Source/Developer/TraceAnalysis/Private/Store/FileStore.cpp b/Engine/Source/Developer/TraceAnalysis/Private/Store/FileStore.cpp
--- a/Engine/Source/Developer/TraceAnalysis/Private/Store/FileStore.cpp
+++ b/Engine/Source/Developer/TraceAnalysis/Private/Store/FileStore.cpp
@@ -58,7 +58,7 @@ private:
 	{
 		FString Name;
 		FString Path;
-		bool IsLive;
+		bool bIsLive;
 	};
 
 	bool IsSessionLive(FSessionHandle Handle) const;
@@ -88,8 +88,8 @@ FFileStoreOutDataStream::~FFileStoreOutDataStream()
 
 void FFileStoreOutDataStream::Write(const void* Data, uint32 Size)
 {
-	bool Success = Inner->Write((const uint8*)Data, Size);
-	check(Success);
+	bool bSuccess = Inner->Write((const uint8*)Data, Size);
+	check(bSuccess);
 }
 
 FFileStoreInDataStream::FFileStoreInDataStream(FFileStore* InOwner, FSessionHandle InSessionHandle, const TCHAR* InFilePath)
@@ -173,7 +173,7 @@ bool FFileStore::GetSessionInfo(FSessionHandle Handle, FSessionInfo& OutInfo)
 	}
 	TSharedPtr<FSessionInfoInternal> Session = AvailableSessions[Handle];
 	OutInfo.Name = *Session->Name;
-	OutInfo.IsLive = Session->IsLive;
+	OutInfo.bIsLive = Session->bIsLive;
 	return true;
 }
 
@@ -190,7 +190,7 @@ IOutDataStream* FFileStore::CreateNewSession()
 	TSharedPtr<FSessionInfoInternal> Session = MakeShared<FSessionInfoInternal>();
 	Session->Name = TraceIdStr;
 	Session->Path = TracePath;
-	Session->IsLive = true;
+	Session->bIsLive = true;
 	AvailableSessions.Add(TraceId, Session);
 
 	return new FFileStoreOutDataStream(this, TraceId, *Session->Path);
@@ -210,13 +210,13 @@ IInDataStream* FFileStore::OpenSessionStream(FSessionHandle Handle)
 bool FFileStore::IsSessionLive(FSessionHandle Handle) const
 {
 	FScopeLock Lock(&SessionsCS);
-	return AvailableSessions[Handle]->IsLive;
+	return AvailableSessions[Handle]->bIsLive;
 }
 
 void FFileStore::CloseSessionStream(FSessionHandle Handle)
 {
 	FScopeLock Lock(&SessionsCS);
-	AvailableSessions[Handle]->IsLive = false;
+	AvailableSessions[Handle]->bIsLive = false;
 }
 
 TSharedPtr<IStore> Store_Create(const TCHAR* StoreDir)
diff --git a/Engine/Source/Developer/TraceAnalysis/Public/Trace/Store.h b/Engine/Source/Developer/TraceAnalysis/Public/Trace/Store.h
--- a/Engine/Source/Developer/TraceAnalysis/Public/Trace/Store.h
+++ b/Engine/Source/Developer/TraceAnalysis/Public/Trace/Store.h
@@ -16,7 +16,7 @@ typedef uint64 FSessionHandle;
 struct FSessionInfo
 {
 	const TCHAR* Name;
-	bool IsLive;
+	bool bIsLive;
 };
 
 class IStore
diff --git a/Engine/Source/Developer/TraceInsights/Private/Insights/InsightsManager.cpp b/Engine/Source/Developer/TraceInsights/Private/Insights/InsightsManager.cpp
--- a/Engine/Source/Developer/TraceInsights/Private/Insights/InsightsManager.cpp
+++ b/Engine/Source/Developer/TraceInsights/Private/Insights/InsightsManager.cpp
@@ -151,7 +151,7 @@ void FInsightsManager::CreateLiveSession()
 	{
 		Trace::FSessionInfo SessionInfo;
 		SessionService->GetSessionInfo(SessionHandle, SessionInfo);
-		if (SessionInfo.IsLive)
+		if (SessionInfo.bIsLive)
 		{
 			TUniquePtr<Trace::IInDataStream> DataStream(SessionService->OpenSessionStream(SessionHandle));
 			check(DataStream);
diff --git a/Engine/Source/Runtime/TraceLog/Private/Trace/Writer.cpp b/Engine/Source/Runtime/TraceLog/Private/Trace/Writer.cpp
--- a/Engine/Source/Runtime/TraceLog/Private/Trace/Writer.cpp
+++ b/Engine/Source/Runtime/TraceLog/Private/Trace/Writer.cpp
@@ -298,7 +298,7 @@ public:
 	void				Init();
 	void				Shutdown();
 	void				Write(const void* Data, uint32 Size);
-	bool				IsFull() const	{ return Full; }
+	bool				IsFull() const	{ return bFull; }
 	const uint8*		GetData() const { return Base; }
 	uint32				GetSize() const { return Used; }
 
@@ -309,7 +309,7 @@ private:
 	uint8*				Base;
 	int32				Used;
 	uint16				MappedPageCount;
-	bool				Full;
+	bool				bFull;
 };
 
 ////////////////////////////////////////////////////////////////////////////////
@@ -318,7 +318,7 @@ void FHoldBuffer::Init()
 	Base = MemoryReserve(FHoldBuffer::PageSize * FHoldBuffer::MaxPages);
 	Used = 0;
 	MappedPageCount = 0;
-	Full = false;
+	bFull = false;
 }
 
 ////////////////////////////////////////////////////////////////////////////////
@@ -345,7 +345,7 @@ void FHoldBuffer::Write(const void* Data, uint32 Size)
 	{
 		if (HotPageCount >= FHoldBuffer::MaxPages)
 		{
-			Full = true;
+			bFull = true;
 			return;
 		}
 
@@ -437,8 +437,8 @@ static void Writer_UpdateData()
 		});
 
 		// Did we overflow? Enter partial mode.
-		bool Overflown = GHoldBuffer.IsFull();
-		if (Overflown && GDataState != EDataState::Partial)
+		bool bOverflown = GHoldBuffer.IsFull();
+		if (bOverflown && GDataState != EDataState::Partial)
 		{
 			/* TODO: turn off non-retail events */
 			GDataState = EDataState::Partial;
