Basic class to wrap connecting to and controlling a running instance's trace system.

diff --git a/Engine/Source/Developer/TraceAnalysis/Private/ControlClient.cpp b/Engine/Source/Developer/TraceAnalysis/Private/ControlClient.cpp
new file mode 100644
--- /dev/null
+++ b/Engine/Source/Developer/TraceAnalysis/Private/ControlClient.cpp
@@ -0,0 +1,134 @@
+// Copyright 1998-2019 Epic Games, Inc. All Rights Reserved.
+
+#include "Trace/ControlClient.h"
+#include "IPAddress.h"
+#include "Misc/CString.h"
+#include "Sockets.h"
+#include "SocketSubsystem.h"
+
+namespace Trace
+{
+
+////////////////////////////////////////////////////////////////////////////////
+FControlClient::~FControlClient()
+{
+    Disconnect();
+}
+
+////////////////////////////////////////////////////////////////////////////////
+bool FControlClient::Connect(const TCHAR* Host, uint16 Port)
+{
+    if (IsConnected())
+    {
+        return false;
+    }
+
+	ISocketSubsystem& Sockets = *(ISocketSubsystem::Get());
+
+	FSocket* ClientSocket = Sockets.CreateSocket(NAME_Stream, TEXT("TraceControlClient"));
+	if (ClientSocket == nullptr)
+	{
+		return false;
+	}
+
+	bool bIsHostValid = false;
+	TSharedPtr<FInternetAddr> Addr = Sockets.CreateInternetAddr();
+	Addr->SetIp(Host, bIsHostValid);
+	Addr->SetPort(Port);
+    if (!bIsHostValid || !ClientSocket->Connect(*Addr))
+    {
+        Sockets.DestroySocket(ClientSocket);
+        return false;
+    }
+
+	ClientSocket->SetLinger();
+
+    Socket = ClientSocket;
+    return true;
+}
+
+////////////////////////////////////////////////////////////////////////////////
+void FControlClient::Disconnect()
+{
+    if (!IsConnected())
+    {
+        return;
+    }
+
+	Socket->Shutdown(ESocketShutdownMode::ReadWrite);
+	Socket->Close();
+
+	ISocketSubsystem& Sockets = *(ISocketSubsystem::Get());
+	Sockets.DestroySocket(Socket);
+	Socket = nullptr;
+}
+
+////////////////////////////////////////////////////////////////////////////////
+bool FControlClient::IsConnected() const
+{
+    return (Socket != nullptr);
+}
+
+////////////////////////////////////////////////////////////////////////////////
+void FControlClient::SendConnect(const TCHAR* Path)
+{
+    if (!IsConnected())
+    {
+        return;
+    }
+
+    FormatAndSend(TEXT("Connect %s"), Path);
+}
+
+////////////////////////////////////////////////////////////////////////////////
+void FControlClient::SendToggleEvent(const TCHAR* EventMask, bool bState)
+{
+    FormatAndSend(TEXT("ToggleEvent %s %d"), EventMask, bState);
+}
+
+////////////////////////////////////////////////////////////////////////////////
+void FControlClient::Send(const TCHAR* Command)
+{
+	int Length = FCString::Strlen(Command);
+	Send((const uint8*)TCHAR_TO_ANSI(Command), Length);
+}
+
+////////////////////////////////////////////////////////////////////////////////
+void FControlClient::FormatAndSend(const TCHAR* Format, ...)
+{
+    if (!IsConnected())
+    {
+        return;
+    }
+
+    TCHAR Buffer[512];
+	va_list Args;
+	va_start(Args, Format);
+	int Length = FCString::GetVarArgs(Buffer, ARRAY_COUNT(Buffer), Format, Args);
+    if (Length > sizeof(Buffer))
+    {
+        Length = sizeof(Buffer);
+    }
+	va_end(Args);
+
+	Send((const uint8*)TCHAR_TO_ANSI(Buffer), Length);
+}
+
+////////////////////////////////////////////////////////////////////////////////
+void FControlClient::Send(const uint8* Data, int Length)
+{
+    int32 SentBytes = 0;
+    if (!Socket->Send(Data, Length, SentBytes) || SentBytes != Length)
+    {
+        Disconnect();
+		return;
+    }
+
+    if (!Socket->Send((const uint8*)"\n", 1, SentBytes) || SentBytes != 1)
+    {
+        Disconnect();
+		return;
+    }
+}
+
+} // namesapce Trace
diff --git a/Engine/Source/Developer/TraceAnalysis/Public/Trace/ControlClient.h b/Engine/Source/Developer/TraceAnalysis/Public/Trace/ControlClient.h
new file mode 100644
--- /dev/null
+++ b/Engine/Source/Developer/TraceAnalysis/Public/Trace/ControlClient.h
@@ -0,0 +1,28 @@
+// Copyright 1998-2019 Epic Games, Inc. All Rights Reserved.
+
+#pragma once
+
+class FSocket;
+
+namespace Trace
+{
+
+////////////////////////////////////////////////////////////////////////////////
+class TRACEANALYSIS_API FControlClient
+{
+public:
+                ~FControlClient();
+    bool        Connect(const TCHAR* Host, uint16 Port=1985);
+    void        Disconnect();
+    bool        IsConnected() const;
+    void        SendConnect(const TCHAR* Path);
+    void        SendToggleEvent(const TCHAR* Logger, bool bState=true);
+    void        Send(const TCHAR* Command);
+
+private:
+    void        FormatAndSend(const TCHAR* Format, ...);
+    void        Send(const uint8* Data, int Length);
+    FSocket*    Socket = nullptr;
+};
+
+} // namespace Trace
