UnrealInsights:
Made the TraceStore more robust against outside interference
More descriptive and unique naming scheme for .utrace files

diff --git a/Engine/Source/Developer/TraceAnalysis/Private/DataStream.cpp b/Engine/Source/Developer/TraceAnalysis/Private/DataStream.cpp
--- a/Engine/Source/Developer/TraceAnalysis/Private/DataStream.cpp
+++ b/Engine/Source/Developer/TraceAnalysis/Private/DataStream.cpp
@@ -19,6 +19,10 @@ FFileStream::~FFileStream()
 
 int32 FFileStream::Read(void* Data, uint32 Size)
 {
+	if (!Inner)
+	{
+		return 0;
+	}
 	uint64 Remaining = End - Cursor;
 	if (Remaining == 0)
 	{
@@ -41,15 +45,20 @@ void FFileStream::UpdateFileSize()
 {
 	delete Inner;
 	OpenFileInternal();
-	Inner->Seek(Cursor);
+	if (Inner)
+	{
+		Inner->Seek(Cursor);
+	}
 }
 
 void FFileStream::OpenFileInternal()
 {
 	IPlatformFile& FileSystem = IPlatformFile::GetPlatformPhysical();
 	Inner = FileSystem.OpenRead(*FilePath, true);
-	check(Inner != nullptr);
-	End = Inner->Size();
+	if (Inner)
+	{
+		End = Inner->Size();
+	}
 }
 
 FStreamReader::FStreamReader(IInDataStream& InDataStream)
@@ -65,6 +74,11 @@ FStreamReader::~FStreamReader()
 
 IInDataStream* DataStream_ReadFile(const TCHAR* FilePath)
 {
+	IPlatformFile& FileSystem = IPlatformFile::GetPlatformPhysical();
+	if (!FileSystem.FileExists(FilePath))
+	{
+		return nullptr;
+	}
 	return new FFileStream(FilePath);
 }
 
diff --git a/Engine/Source/Developer/TraceAnalysis/Private/Recorder/Recorder.cpp b/Engine/Source/Developer/TraceAnalysis/Private/Recorder/Recorder.cpp
--- a/Engine/Source/Developer/TraceAnalysis/Private/Recorder/Recorder.cpp
+++ b/Engine/Source/Developer/TraceAnalysis/Private/Recorder/Recorder.cpp
@@ -68,7 +68,11 @@ uint32 FRecorder::FSession::Run()
 			break;
 		}
 
-		StoreSession->Write(Buffer, RecvSize);
+		if (!StoreSession->Write(Buffer, RecvSize))
+		{
+			bDead = true;
+			break;
+		}
 	}
 	while (!bDead);
 
@@ -109,9 +113,14 @@ FRecorder::FSession* FRecorder::AcceptSession(FSocket& Socket)
 		return nullptr;
 	}
 
+	IOutDataStream* StoreSession = Store->CreateNewSession();
+	if (!StoreSession)
+	{
+		return nullptr;
+	}
 	FSession* Session = new FSession();
 	Session->Socket = &Socket;
-	Session->StoreSession = Store->CreateNewSession();
+	Session->StoreSession = StoreSession;
 	Session->Thread = FRunnableThread::Create(Session, TEXT("TraceRecSession"));
 	return Session;
 }
diff --git a/Engine/Source/Developer/TraceAnalysis/Private/Store/FileStore.cpp b/Engine/Source/Developer/TraceAnalysis/Private/Store/FileStore.cpp
--- a/Engine/Source/Developer/TraceAnalysis/Private/Store/FileStore.cpp
+++ b/Engine/Source/Developer/TraceAnalysis/Private/Store/FileStore.cpp
@@ -6,6 +6,11 @@
 #include "HAL/PlatformProcess.h"
 #include "Misc/ScopeLock.h"
 #include "Templates/UniquePtr.h"
+#include "Misc/Paths.h"
+#include "Modules/ModuleManager.h"
+#include "DirectoryWatcherModule.h"
+#include "IDirectoryWatcher.h"
+#include "Containers/Ticker.h"
 
 namespace Trace
 {
@@ -16,7 +21,7 @@ class FFileStoreInDataStream
 	: public IInDataStream
 {
 public:
-	FFileStoreInDataStream(FFileStore* Owner, FSessionHandle SessionHandle, const TCHAR* FilePath);
+	FFileStoreInDataStream(FFileStore* Owner, FSessionHandle SessionHandle, FFileStream* File);
 	virtual int32 Read(void* Data, uint32 Size) override;
 
 private:
@@ -29,9 +34,9 @@ class FFileStoreOutDataStream
 	: public IOutDataStream
 {
 public:
-	FFileStoreOutDataStream(FFileStore* Owner, FSessionHandle SessionHandle, const TCHAR* FilePath);
+	FFileStoreOutDataStream(FFileStore* Owner, FSessionHandle SessionHandle, IFileHandle* File);
 	virtual ~FFileStoreOutDataStream();
-	virtual void Write(const void* Data, uint32 Size) override;
+	virtual bool Write(const void* Data, uint32 Size) override;
 
 private:
 	FFileStore* Owner;
@@ -44,6 +49,7 @@ class FFileStore
 {
 public:
 	FFileStore(const TCHAR* StoreDir);
+	~FFileStore();
 
 	virtual void GetAvailableSessions(TArray<FSessionHandle>& OutTraces) override;
 	virtual bool GetSessionInfo(FSessionHandle Handle, FSessionInfo& OutInfo) override;
@@ -56,28 +62,37 @@ private:
 
 	struct FSessionInfoInternal
 	{
+		FSessionHandle Handle;
 		FString Name;
 		FString Path;
 		bool bIsLive;
+		bool bIsValid;
 	};
 
+	FSessionInfoInternal* AddSession(const FString& Path);
+	void RemoveSession(FSessionHandle Handle);
+	FSessionInfoInternal* GetSession(FSessionHandle Handle) const;
+	bool Tick(float DeltaTime);
+	void OnDirectoryChanged(const TArray<FFileChangeData>& FileChanges);
 	bool IsSessionLive(FSessionHandle Handle) const;
 	void CloseSessionStream(FSessionHandle Handle);
 
 	mutable FCriticalSection SessionsCS;
 	FString StoreDir;
-	uint32 NextTraceId;
-	TMap<FSessionHandle, TSharedPtr<FSessionInfoInternal>> AvailableSessions;
+	IDirectoryWatcher* DirectoryWatcher;
+	FDelegateHandle DirectoryWatcherHandle;
+	FDelegateHandle TickHandle;
+	FSessionHandle NextSessionHandle = 1;
+	TArray<FSessionInfoInternal*> Sessions;
+	TMap<FString, FSessionInfoInternal*> SessionsByPathMap;
 };
 
-FFileStoreOutDataStream::FFileStoreOutDataStream(FFileStore* InOwner, FSessionHandle InSessionHandle, const TCHAR* InFilePath)
+FFileStoreOutDataStream::FFileStoreOutDataStream(FFileStore* InOwner, FSessionHandle InSessionHandle, IFileHandle* InFile)
 	: Owner(InOwner)
 	, SessionHandle(InSessionHandle)
 {
-	IPlatformFile& FileSystem = IPlatformFile::GetPlatformPhysical();
-	IFileHandle* File = FileSystem.OpenWrite(InFilePath, true, true);
-	check(File != nullptr);
-	Inner.Reset(File);
+	check(InFile);
+	Inner.Reset(InFile);
 }
 
 FFileStoreOutDataStream::~FFileStoreOutDataStream()
@@ -86,16 +101,15 @@ FFileStoreOutDataStream::~FFileStoreOutDataStream()
 	Owner->CloseSessionStream(SessionHandle);
 }
 
-void FFileStoreOutDataStream::Write(const void* Data, uint32 Size)
+bool FFileStoreOutDataStream::Write(const void* Data, uint32 Size)
 {
-	bool bSuccess = Inner->Write((const uint8*)Data, Size);
-	check(bSuccess);
+	return Inner->Write((const uint8*)Data, Size);
 }
 
-FFileStoreInDataStream::FFileStoreInDataStream(FFileStore* InOwner, FSessionHandle InSessionHandle, const TCHAR* InFilePath)
+FFileStoreInDataStream::FFileStoreInDataStream(FFileStore* InOwner, FSessionHandle InSessionHandle, FFileStream* InFileStream)
 	: Owner(InOwner)
 	, SessionHandle(InSessionHandle)
-	, Inner(new FFileStream(InFilePath))
+	, Inner(InFileStream)
 {
 }
 
@@ -128,50 +142,109 @@ FFileStore::FFileStore(const TCHAR* InStoreDir)
 	FileSystem.CreateDirectory(*StoreDir);
 
 	int32 HighestTraceId = -1;
-	FileSystem.IterateDirectory(*StoreDir, [&HighestTraceId, this](const TCHAR* FileName, bool)
+	FileSystem.IterateDirectory(*StoreDir, [&HighestTraceId, this](const TCHAR* FileName, bool IsDirectory)
 	{
-		if (const TCHAR* Dot = FCString::Strrchr(FileName, '.'))
+		if (!IsDirectory && FPaths::GetExtension(FileName) == TEXT("utrace"))
 		{
-			if (FCString::Stricmp(Dot, TEXT(".utrace")) == 0)
-			{
-				const TCHAR* Backslash = FCString::Strrchr(FileName, '\\');
-				const TCHAR* Forwardslash = FCString::Strrchr(FileName, '/');
-				const TCHAR* Slash = UPTRINT(Backslash) > UPTRINT(Forwardslash) ? Backslash : Forwardslash;
-				if (Slash != nullptr)
-				{
-					int32 Id = FCString::Atoi(Slash + 1);
-					TSharedPtr<FSessionInfoInternal> Session = MakeShared<FSessionInfoInternal>();
-					Session->Name = FString::Printf(TEXT("%04d"), Id);
-					Session->Path = FileName;
-					AvailableSessions.Add(Id, Session);
-					HighestTraceId = (Id > HighestTraceId) ? Id : HighestTraceId;
-				}
-			}
+			AddSession(FileName);
 		}
 		return true;
 	});
 
-	NextTraceId = HighestTraceId + 1;
+	FDirectoryWatcherModule& DirectoryWatcherModule = FModuleManager::LoadModuleChecked<FDirectoryWatcherModule>(TEXT("DirectoryWatcher"));
+	DirectoryWatcher = DirectoryWatcherModule.Get();
+	if (DirectoryWatcher)
+	{
+		DirectoryWatcher->RegisterDirectoryChangedCallback_Handle(StoreDir, IDirectoryWatcher::FDirectoryChanged::CreateRaw(this, &FFileStore::OnDirectoryChanged), DirectoryWatcherHandle, IDirectoryWatcher::WatchOptions::IncludeDirectoryChanges);
+		FTickerDelegate TickDelegate = FTickerDelegate::CreateRaw(this, &FFileStore::Tick);
+		TickHandle = FTicker::GetCoreTicker().AddTicker(TickDelegate, 0.5f);
+	}
+}
+
+FFileStore::~FFileStore()
+{
+	FDirectoryWatcherModule* DirectoryWatcherModule = FModuleManager::GetModulePtr<FDirectoryWatcherModule>(TEXT("DirectoryWatcher"));
+	if (DirectoryWatcherModule)
+	{
+		IDirectoryWatcher* LocalDirectoryWatcher = DirectoryWatcherModule->Get();
+		if (DirectoryWatcher)
+		{
+			LocalDirectoryWatcher->UnregisterDirectoryChangedCallback_Handle(StoreDir, DirectoryWatcherHandle);
+		}
+	}
+	if (TickHandle.IsValid())
+	{
+		FTicker::GetCoreTicker().RemoveTicker(TickHandle);
+		TickHandle.Reset();
+	}
+	for (FSessionInfoInternal* Session : Sessions)
+	{
+		delete Session;
+	}
+}
+
+FFileStore::FSessionInfoInternal* FFileStore::AddSession(const FString& Path)
+{
+	FScopeLock Lock(&SessionsCS);
+
+	FSessionInfoInternal* Session = new FSessionInfoInternal();
+	Session->Handle = NextSessionHandle++;
+	Session->Name = FPaths::GetBaseFilename(Path);
+	Session->Path = Path;
+	Session->bIsValid = true;
+	Sessions.Add(Session);
+	SessionsByPathMap.Add(Path, Session);
+
+	return Session;
+}
+
+void FFileStore::RemoveSession(FSessionHandle Handle)
+{
+	FScopeLock Lock(&SessionsCS);
+
+	FSessionInfoInternal* Session = Sessions[Handle - 1];
+	SessionsByPathMap.Remove(Session->Path);
+	Sessions[Handle - 1]->bIsValid = false;
+}
+
+FFileStore::FSessionInfoInternal* FFileStore::GetSession(FSessionHandle Handle) const
+{
+	FScopeLock Lock(&SessionsCS);
+	if (Handle - 1 >= Sessions.Num())
+	{
+		return nullptr;
+	}
+	FSessionInfoInternal* Session = Sessions[Handle - 1];
+	if (!Session->bIsValid)
+	{
+		return nullptr;
+	}
+	return Session;
 }
 
 void FFileStore::GetAvailableSessions(TArray<FSessionHandle>& OutTraces)
 {
 	FScopeLock Lock(&SessionsCS);
-	OutTraces.Empty(AvailableSessions.Num());
-	for (const auto& KeyValuePair : AvailableSessions)
+	OutTraces.Empty(Sessions.Num());
+	for (const FSessionInfoInternal* Session : Sessions)
 	{
-		OutTraces.Add(KeyValuePair.Key);
+		if (Session->bIsValid)
+		{
+			OutTraces.Add(Session->Handle);
+		}
 	}
 }
 
 bool FFileStore::GetSessionInfo(FSessionHandle Handle, FSessionInfo& OutInfo)
 {
 	FScopeLock Lock(&SessionsCS);
-	if (!AvailableSessions.Contains(Handle))
+	
+	FSessionInfoInternal* Session = GetSession(Handle);
+	if (!Session)
 	{
 		return false;
 	}
-	TSharedPtr<FSessionInfoInternal> Session = AvailableSessions[Handle];
+	OutInfo.Uri = *Session->Path;
 	OutInfo.Name = *Session->Name;
 	OutInfo.bIsLive = Session->bIsLive;
 	return true;
@@ -181,42 +254,99 @@ IOutDataStream* FFileStore::CreateNewSession()
 {
 	FScopeLock Lock(&SessionsCS);
 
-	uint32 TraceId = NextTraceId++;
-	FString TraceIdStr = FString::Printf(TEXT("%04d"), TraceId);
-	FString TracePath = StoreDir;
-	TracePath /= TraceIdStr;
-	TracePath += TEXT(".utrace");
+	IPlatformFile& FileSystem = IPlatformFile::GetPlatformPhysical();
+
+	FString TracePath = StoreDir / FString::Printf(TEXT("%s.utrace"), *FDateTime::Now().ToString(TEXT("%Y%m%d_%H%M%S")));
+	uint64 Index = 1;
+	while (FileSystem.FileExists(*TracePath))
+	{
+		TracePath = StoreDir / FString::Printf(TEXT("%s_%d.utrace"), *FDateTime::Now().ToString(TEXT("%Y%m%d_%H%M%S")), Index);
+	}
+
+	IFileHandle* File = FileSystem.OpenWrite(*TracePath, true, true);
+	if (!File)
+	{
+		return nullptr;
+	}
 
-	TSharedPtr<FSessionInfoInternal> Session = MakeShared<FSessionInfoInternal>();
-	Session->Name = TraceIdStr;
-	Session->Path = TracePath;
+	FSessionInfoInternal* Session = AddSession(TracePath);
 	Session->bIsLive = true;
-	AvailableSessions.Add(TraceId, Session);
 
-	return new FFileStoreOutDataStream(this, TraceId, *Session->Path);
+	return new FFileStoreOutDataStream(this, Session->Handle, File);
 }
 
 IInDataStream* FFileStore::OpenSessionStream(FSessionHandle Handle)
 {
 	FScopeLock Lock(&SessionsCS);
-	if (!AvailableSessions.Contains(Handle))
+	
+	FSessionInfoInternal* Session = GetSession(Handle);
+	if (!Session)
 	{
 		return nullptr;
 	}
-	TSharedPtr<FSessionInfoInternal> Session = AvailableSessions[Handle];
-	return new FFileStoreInDataStream(this, Handle, *Session->Path);
+
+	IPlatformFile& FileSystem = IPlatformFile::GetPlatformPhysical();
+	if (!FileSystem.FileExists(*Session->Path))
+	{
+		return nullptr;
+	}
+
+	FFileStream* FileStream = new FFileStream(*Session->Path);
+	return new FFileStoreInDataStream(this, Handle, FileStream);
+}
+
+bool FFileStore::Tick(float DeltaTime)
+{
+	DirectoryWatcher->Tick(DeltaTime);
+	return true;
+}
+
+void FFileStore::OnDirectoryChanged(const TArray<FFileChangeData>& FileChanges)
+{
+	FScopeLock Lock(&SessionsCS);
+
+	for (const FFileChangeData& FileChange : FileChanges)
+	{
+		FSessionInfoInternal** FindIt = SessionsByPathMap.Find(FileChange.Filename.ToUpper());
+		switch (FileChange.Action)
+		{
+		case FFileChangeData::FCA_Removed:
+			if (FindIt)
+			{
+				RemoveSession((*FindIt)->Handle);
+			}
+			break;
+		case FFileChangeData::FCA_Added:
+			if (!FindIt && FPaths::GetExtension(FileChange.Filename) == TEXT("utrace"))
+			{
+				AddSession(FileChange.Filename);
+			}
+			break;
+		}
+	}
 }
 
 bool FFileStore::IsSessionLive(FSessionHandle Handle) const
 {
 	FScopeLock Lock(&SessionsCS);
-	return AvailableSessions[Handle]->bIsLive;
+	
+	FSessionInfoInternal* Session = GetSession(Handle);
+	if (!Session)
+	{
+		return false;
+	}
+	return Session->bIsLive;
 }
 
 void FFileStore::CloseSessionStream(FSessionHandle Handle)
 {
 	FScopeLock Lock(&SessionsCS);
-	AvailableSessions[Handle]->bIsLive = false;
+	
+	FSessionInfoInternal* Session = GetSession(Handle);
+	if (Session)
+	{
+		Session->bIsLive = false;
+	}
 }
 
 TSharedPtr<IStore> Store_Create(const TCHAR* StoreDir)
diff --git a/Engine/Source/Developer/TraceAnalysis/Public/Trace/DataStream.h b/Engine/Source/Developer/TraceAnalysis/Public/Trace/DataStream.h
--- a/Engine/Source/Developer/TraceAnalysis/Public/Trace/DataStream.h
+++ b/Engine/Source/Developer/TraceAnalysis/Public/Trace/DataStream.h
@@ -18,7 +18,7 @@ class IOutDataStream
 {
 public:
 	virtual			~IOutDataStream() = default;
-	virtual void	Write(const void* Data, uint32 Size) = 0;
+	virtual bool	Write(const void* Data, uint32 Size) = 0;
 };
 
 TRACEANALYSIS_API IInDataStream* DataStream_ReadFile(const TCHAR* FilePath);
diff --git a/Engine/Source/Developer/TraceAnalysis/Public/Trace/Store.h b/Engine/Source/Developer/TraceAnalysis/Public/Trace/Store.h
--- a/Engine/Source/Developer/TraceAnalysis/Public/Trace/Store.h
+++ b/Engine/Source/Developer/TraceAnalysis/Public/Trace/Store.h
@@ -15,6 +15,7 @@ typedef uint64 FSessionHandle;
 
 struct FSessionInfo
 {
+	const TCHAR* Uri;
 	const TCHAR* Name;
 	bool bIsLive;
 };
diff --git a/Engine/Source/Developer/TraceAnalysis/TraceAnalysis.Build.cs b/Engine/Source/Developer/TraceAnalysis/TraceAnalysis.Build.cs
--- a/Engine/Source/Developer/TraceAnalysis/TraceAnalysis.Build.cs
+++ b/Engine/Source/Developer/TraceAnalysis/TraceAnalysis.Build.cs
@@ -8,5 +8,6 @@ public class TraceAnalysis : ModuleRules
 	{
 		PrivateDependencyModuleNames.Add("Core");
 		PrivateDependencyModuleNames.Add("Sockets");
+		PrivateDependencyModuleNames.Add("DirectoryWatcher");
 	}
 }
diff --git a/Engine/Source/Developer/TraceInsights/Private/Insights/InsightsManager.cpp b/Engine/Source/Developer/TraceInsights/Private/Insights/InsightsManager.cpp
--- a/Engine/Source/Developer/TraceInsights/Private/Insights/InsightsManager.cpp
+++ b/Engine/Source/Developer/TraceInsights/Private/Insights/InsightsManager.cpp
@@ -263,11 +263,11 @@ void FInsightsManager::LoadSession(Trace::FSessionHandle SessionHandle)
 	SessionService->GetSessionInfo(SessionHandle, SessionInfo);
 
 	TUniquePtr<Trace::IInDataStream> DataStream(SessionService->OpenSessionStream(SessionHandle));
-	check(DataStream);
-
-	Session = AnalysisService->StartAnalysis(SessionInfo.Name, MoveTemp(DataStream));
-
-	SpawnAndActivateTabs();
+	if (DataStream)
+	{
+		Session = AnalysisService->StartAnalysis(SessionInfo.Name, MoveTemp(DataStream));
+		SpawnAndActivateTabs();
+	}
 	OnSessionChanged();
 }
 
@@ -278,11 +278,11 @@ void FInsightsManager::LoadTraceFile(const FString& TraceFilepath)
 	ResetSession();
 
 	TUniquePtr<Trace::IInDataStream> DataStream(SessionService->OpenSessionFromFile(*TraceFilepath));
-	check(DataStream);
-
-	Session = AnalysisService->StartAnalysis(*TraceFilepath, MoveTemp(DataStream));
-
-	SpawnAndActivateTabs();
+	if (DataStream)
+	{
+		Session = AnalysisService->StartAnalysis(*TraceFilepath, MoveTemp(DataStream));
+		SpawnAndActivateTabs();
+	}
 	OnSessionChanged();
 }
 
diff --git a/Engine/Source/Developer/TraceServices/Private/AnalysisService.cpp b/Engine/Source/Developer/TraceServices/Private/AnalysisService.cpp
--- a/Engine/Source/Developer/TraceServices/Private/AnalysisService.cpp
+++ b/Engine/Source/Developer/TraceServices/Private/AnalysisService.cpp
@@ -61,7 +61,8 @@ void FAnalysisSessionLock::EndEdit()
 }
 
 FAnalysisSession::FAnalysisSession(const TCHAR* SessionName)
-	: DurationSeconds(0.0)
+	: Name(SessionName)
+	, DurationSeconds(0.0)
 	, Allocator(32 << 20)
 	, StringStore(Allocator)
 	, BookmarkProvider(Lock, StringStore)
@@ -89,7 +90,7 @@ FAnalysisService::FAnalysisWorker::FAnalysisWorker(FAnalysisService& InOuter, TU
 	
 }
 
-void FAnalysisService::AnalyzeInternal(TSharedRef<FAnalysisSession> AnalysisSession, Trace::IInDataStream& DataStream)
+void FAnalysisService::AnalyzeInternal(TSharedRef<FAnalysisSession> AnalysisSession, Trace::IInDataStream* DataStream)
 {
 	OnAnalysisStarted().Broadcast(AnalysisSession);
 
@@ -112,7 +113,7 @@ void FAnalysisService::AnalyzeInternal(TSharedRef<FAnalysisSession> AnalysisSess
 	}
 	Trace::FAnalysisProcessor Processor = Context.Process();
 
-	Processor.Start(DataStream);
+	Processor.Start(*DataStream);
 
 	for (Trace::IAnalyzer* Analyzer : Analyzers)
 	{
@@ -120,6 +121,8 @@ void FAnalysisService::AnalyzeInternal(TSharedRef<FAnalysisSession> AnalysisSess
 	}
 	Analyzers.Empty();
 
+	delete DataStream;
+
 	AnalysisSession->SetComplete();
 
 	OnAnalysisFinished().Broadcast(AnalysisSession);
@@ -127,7 +130,7 @@ void FAnalysisService::AnalyzeInternal(TSharedRef<FAnalysisSession> AnalysisSess
 
 void FAnalysisService::FAnalysisWorker::DoWork()
 {
-	Outer.AnalyzeInternal(AnalysisSession.ToSharedRef(), *DataStream.Get());
+	Outer.AnalyzeInternal(AnalysisSession.ToSharedRef(), DataStream.Release());
 	AnalysisSession = nullptr;
 }
 
@@ -143,7 +146,7 @@ TSharedPtr<const IAnalysisSession> FAnalysisService::Analyze(const TCHAR* Sessio
 {
 	TSharedRef<FAnalysisSession> AnalysisSession = MakeShared<FAnalysisSession>(SessionName);
 	TUniquePtr<Trace::IInDataStream> DataStream = MoveTemp(InDataStream);
-	AnalyzeInternal(AnalysisSession, *DataStream.Get());
+	AnalyzeInternal(AnalysisSession, DataStream.Release());
 	return AnalysisSession;
 }
 
diff --git a/Engine/Source/Developer/TraceServices/Private/AnalysisServicePrivate.h b/Engine/Source/Developer/TraceServices/Private/AnalysisServicePrivate.h
--- a/Engine/Source/Developer/TraceServices/Private/AnalysisServicePrivate.h
+++ b/Engine/Source/Developer/TraceServices/Private/AnalysisServicePrivate.h
@@ -144,7 +144,7 @@ private:
 		TSharedPtr<FAnalysisSession> AnalysisSession;
 	};
 
-	void AnalyzeInternal(TSharedRef<FAnalysisSession> AnalysisSession, Trace::IInDataStream& DataStream);
+	void AnalyzeInternal(TSharedRef<FAnalysisSession> AnalysisSession, Trace::IInDataStream* DataStream);
 
 	FAnalysisStartedEvent AnalysisStartedEvent;
 	FAnalysisFinishedEvent AnalysisFinishedEvent;
