The analysis engine was using the wrong event identifier when receiving declarations of new events.

diff --git a/Engine/Source/Developer/TraceAnalysis/Private/Analysis/Engine.cpp b/Engine/Source/Developer/TraceAnalysis/Private/Analysis/Engine.cpp
--- a/Engine/Source/Developer/TraceAnalysis/Private/Analysis/Engine.cpp
+++ b/Engine/Source/Developer/TraceAnalysis/Private/Analysis/Engine.cpp
@@ -467,7 +467,7 @@ void FAnalysisEngine::OnNewEvent(const FOnEventContext& Context)
 {
 	const auto& NewEvent = *(FNewEventEvent*)Context.EventData.GetData();
 
-	FDispatch& Dispatch = AddDispatch(NewEvent.Uid, NewEvent.FieldCount);
+	FDispatch& Dispatch = AddDispatch(NewEvent.EventUid, NewEvent.FieldCount);
 
 	if (NewEvent.FieldCount)
 	{
