UnrealInsights:
Added functionality to toggle events on or off for a running recorder session

diff --git a/Engine/Source/Developer/TraceAnalysis/Private/ControlClient.cpp b/Engine/Source/Developer/TraceAnalysis/Private/ControlClient.cpp
--- a/Engine/Source/Developer/TraceAnalysis/Private/ControlClient.cpp
+++ b/Engine/Source/Developer/TraceAnalysis/Private/ControlClient.cpp
@@ -23,6 +23,30 @@ bool FControlClient::Connect(const TCHAR* Host, uint16 Port)
         return false;
     }
 
+	ISocketSubsystem& Sockets = *(ISocketSubsystem::Get());
+	TSharedPtr<FInternetAddr> Addr = Sockets.CreateInternetAddr();
+	bool bIsHostValid = false;
+	Addr->SetIp(Host, bIsHostValid);
+	if (!bIsHostValid)
+	{
+		ESocketErrors SocketErrors = Sockets.GetHostByName(TCHAR_TO_ANSI(Host), *Addr);
+		if (SocketErrors != SE_NO_ERROR)
+		{
+			return false;
+		}
+	}
+	Addr->SetPort(Port);
+	return Connect(*Addr);
+}
+
+////////////////////////////////////////////////////////////////////////////////
+bool FControlClient::Connect(const FInternetAddr& Address)
+{
+	if (IsConnected())
+	{
+		return false;
+	}
+
 	ISocketSubsystem& Sockets = *(ISocketSubsystem::Get());
 
 	FSocket* ClientSocket = Sockets.CreateSocket(NAME_Stream, TEXT("TraceControlClient"));
@@ -31,15 +55,11 @@ bool FControlClient::Connect(const TCHAR* Host, uint16 Port)
 		return false;
 	}
 
-	bool bIsHostValid = false;
-	TSharedPtr<FInternetAddr> Addr = Sockets.CreateInternetAddr();
-	Addr->SetIp(Host, bIsHostValid);
-	Addr->SetPort(Port);
-    if (!bIsHostValid || !ClientSocket->Connect(*Addr))
-    {
-        Sockets.DestroySocket(ClientSocket);
-        return false;
-    }
+	if (!ClientSocket->Connect(Address))
+	{
+		Sockets.DestroySocket(ClientSocket);
+		return false;
+	}
 
 	ClientSocket->SetLinger();
 
diff --git a/Engine/Source/Developer/TraceAnalysis/Private/Recorder/Recorder.cpp b/Engine/Source/Developer/TraceAnalysis/Private/Recorder/Recorder.cpp
--- a/Engine/Source/Developer/TraceAnalysis/Private/Recorder/Recorder.cpp
+++ b/Engine/Source/Developer/TraceAnalysis/Private/Recorder/Recorder.cpp
@@ -10,6 +10,7 @@
 #include "SocketSubsystem.h"
 #include "Trace/DataStream.h"
 #include "Trace/Store.h"
+#include "Trace/ControlClient.h"
 
 namespace Trace
 {
@@ -30,13 +31,17 @@ private:
 	virtual bool		StartRecording() override;
 	virtual void		StopRecording() override;
 	virtual uint32		GetSessionCount() const override;
+	virtual void		GetActiveSessions(TArray<FRecorderSessionInfo>& OutSessions) const override;
+	virtual bool		ToggleEvent(FRecorderSessionHandle RecordingHandle, const TCHAR* LoggerWildcard, bool bState) override;
 	FSession*			AcceptSession(FSocket& Socket);
 	void				CloseSession(FSession& Session);
 	void				ReapDeadSessions();
 	TSharedRef<IStore>	Store;
-	TArray<FSession*>	Sessions;
+	mutable FCriticalSection SessionsCS;
+	TMap<FRecorderSessionHandle, FSession*> Sessions;
 	FRunnableThread*	Thread = nullptr;
 	FSocket*			ListenSocket = nullptr;
+	FRecorderSessionHandle	NextSessionHandle = 1;
 	volatile bool		bStopRequested;
 };
 
@@ -48,8 +53,10 @@ struct FRecorder::FSession
 {
 	virtual uint32		Run() override;
 	FSocket*			Socket;
-	IOutDataStream*		StoreSession;
+	FStoreSessionHandle	StoreSessionHandle;
+	IOutDataStream*		StoreSessionStream;
 	FRunnableThread*	Thread;
+	TSharedPtr<FInternetAddr> ControlClientAddress;
 	volatile bool		bDead = false;
 };
 
@@ -68,7 +75,7 @@ uint32 FRecorder::FSession::Run()
 			break;
 		}
 
-		if (!StoreSession->Write(Buffer, RecvSize))
+		if (!StoreSessionStream->Write(Buffer, RecvSize))
 		{
 			bDead = true;
 			break;
@@ -113,14 +120,18 @@ FRecorder::FSession* FRecorder::AcceptSession(FSocket& Socket)
 		return nullptr;
 	}
 
-	IOutDataStream* StoreSession = Store->CreateNewSession();
-	if (!StoreSession)
+	TTuple<FStoreSessionHandle, IOutDataStream*> StoreSession = Store->CreateNewSession();
+	if (!StoreSession.Get<1>())
 	{
 		return nullptr;
 	}
 	FSession* Session = new FSession();
+	Session->ControlClientAddress = ISocketSubsystem::Get(PLATFORM_SOCKETSUBSYSTEM)->CreateInternetAddr();
+	Socket.GetPeerAddress(*Session->ControlClientAddress);
+	Session->ControlClientAddress->SetPort(1985);
 	Session->Socket = &Socket;
-	Session->StoreSession = StoreSession;
+	Session->StoreSessionHandle = StoreSession.Get<0>();
+	Session->StoreSessionStream = StoreSession.Get<1>();
 	Session->Thread = FRunnableThread::Create(Session, TEXT("TraceRecSession"));
 	return Session;
 }
@@ -137,7 +148,7 @@ void FRecorder::CloseSession(FSession& Session)
 	Sockets.DestroySocket(Session.Socket);
 
 	delete Session.Thread;
-	delete Session.StoreSession;
+	delete Session.StoreSessionStream;
 	delete &Session;
 }
 
@@ -145,21 +156,24 @@ void FRecorder::CloseSession(FSession& Session)
 void FRecorder::ReapDeadSessions()
 {
 	// Reap dead sessions
-	int SessionCount = Sessions.Num();
-	for (int i = 0; i < SessionCount; ++i)
+	TArray<FSession*> SessionsToClose;
 	{
-		FSession* Session = Sessions[i];
-		if (!Session->bDead)
+		FScopeLock Lock(&SessionsCS);
+		for (auto It = Sessions.CreateIterator(); It; ++It)
 		{
-			continue;
+			FSession* Session = It->Value;
+			if (Session->bDead)
+			{
+				SessionsToClose.Add(Session);
+				It.RemoveCurrent();
+			}
 		}
-
+	}
+	
+	for (FSession* Session : SessionsToClose)
+	{
 		CloseSession(*Session);
-
-		--SessionCount;
-		Sessions[i] = Sessions[SessionCount];
 	}
-	Sessions.SetNumUnsafeInternal(SessionCount);
 }
 
 ////////////////////////////////////////////////////////////////////////////////
@@ -186,7 +200,8 @@ uint32 FRecorder::Run()
 
 			if (FSession* Session = AcceptSession(*Socket))
 			{
-				Sessions.Add(Session);
+				FScopeLock Lock(&SessionsCS);
+				Sessions.Add(NextSessionHandle++, Session);
 			}
 			else
 			{
@@ -197,7 +212,16 @@ uint32 FRecorder::Run()
 		ReapDeadSessions();
 	}
 
-	for (FSession* Session : Sessions)
+	TArray<FSession*> SessionsToClose;
+	{
+		FScopeLock Lock(&SessionsCS);
+		for (auto& KV : Sessions)
+		{
+			SessionsToClose.Add(KV.Value);
+		}
+		Sessions.Empty();
+	}
+	for (FSession* Session : SessionsToClose)
 	{
 		CloseSession(*Session);
 	}
@@ -276,10 +300,46 @@ void FRecorder::StopRecording()
 ////////////////////////////////////////////////////////////////////////////////
 uint32 FRecorder::GetSessionCount() const
 {
+	FScopeLock Lock(&SessionsCS);
 	return Sessions.Num();
 }
 
 
+////////////////////////////////////////////////////////////////////////////////
+void FRecorder::GetActiveSessions(TArray<FRecorderSessionInfo>& OutSessions) const
+{
+	FScopeLock Lock(&SessionsCS);
+	OutSessions.Reserve(OutSessions.Num() + Sessions.Num());
+	for (const auto& KV : Sessions)
+	{
+		FRecorderSessionInfo& SessionInfo = OutSessions.AddDefaulted_GetRef();
+		SessionInfo.Handle = KV.Key;
+		SessionInfo.StoreSessionHandle = KV.Value->StoreSessionHandle;
+	}
+}
+
+////////////////////////////////////////////////////////////////////////////////
+bool FRecorder::ToggleEvent(FRecorderSessionHandle SessionHandle, const TCHAR* LoggerWildcard, bool bState)
+{
+	TSharedPtr<FInternetAddr> ControlClientAddress;
+	{
+		FScopeLock Lock(&SessionsCS);
+		FSession** FindIt = Sessions.Find(SessionHandle);
+		if (!FindIt)
+		{
+			return false;
+		}
+		ControlClientAddress = (*FindIt)->ControlClientAddress;
+	}
+	FControlClient ControlClient;
+	if (!ControlClient.Connect(*ControlClientAddress))
+	{
+		return false;
+	}
+	ControlClient.SendToggleEvent(LoggerWildcard, bState);
+	ControlClient.Disconnect();
+	return true;
+}
 
 ////////////////////////////////////////////////////////////////////////////////
 TSharedPtr<IRecorder> Recorder_Create(TSharedRef<IStore> Store)
diff --git a/Engine/Source/Developer/TraceAnalysis/Private/Store/FileStore.cpp b/Engine/Source/Developer/TraceAnalysis/Private/Store/FileStore.cpp
--- a/Engine/Source/Developer/TraceAnalysis/Private/Store/FileStore.cpp
+++ b/Engine/Source/Developer/TraceAnalysis/Private/Store/FileStore.cpp
@@ -21,12 +21,12 @@ class FFileStoreInDataStream
 	: public IInDataStream
 {
 public:
-	FFileStoreInDataStream(FFileStore* Owner, FSessionHandle SessionHandle, FFileStream* File);
+	FFileStoreInDataStream(FFileStore* Owner, FStoreSessionHandle SessionHandle, FFileStream* File);
 	virtual int32 Read(void* Data, uint32 Size) override;
 
 private:
 	FFileStore* Owner;
-	FSessionHandle SessionHandle;
+	FStoreSessionHandle SessionHandle;
 	TUniquePtr<FFileStream> Inner;
 };
 
@@ -34,13 +34,13 @@ class FFileStoreOutDataStream
 	: public IOutDataStream
 {
 public:
-	FFileStoreOutDataStream(FFileStore* Owner, FSessionHandle SessionHandle, IFileHandle* File);
+	FFileStoreOutDataStream(FFileStore* Owner, FStoreSessionHandle SessionHandle, IFileHandle* File);
 	virtual ~FFileStoreOutDataStream();
 	virtual bool Write(const void* Data, uint32 Size) override;
 
 private:
 	FFileStore* Owner;
-	FSessionHandle SessionHandle;
+	FStoreSessionHandle SessionHandle;
 	TUniquePtr<IFileHandle> Inner;
 };
 
@@ -51,10 +51,9 @@ public:
 	FFileStore(const TCHAR* StoreDir);
 	~FFileStore();
 
-	virtual void GetAvailableSessions(TArray<FSessionHandle>& OutTraces) override;
-	virtual bool GetSessionInfo(FSessionHandle Handle, FSessionInfo& OutInfo) override;
-	virtual IOutDataStream* CreateNewSession() override;
-	virtual IInDataStream* OpenSessionStream(FSessionHandle Handle) override;
+	virtual void GetAvailableSessions(TArray<FStoreSessionInfo>& OutSessions) const override;
+	virtual TTuple<FStoreSessionHandle, IOutDataStream*> CreateNewSession() override;
+	virtual IInDataStream* OpenSessionStream(FStoreSessionHandle Handle) override;
 
 private:
 	friend class FFileStoreInDataStream;
@@ -62,7 +61,7 @@ private:
 
 	struct FSessionInfoInternal
 	{
-		FSessionHandle Handle;
+		FStoreSessionHandle Handle;
 		FString Name;
 		FString Path;
 		bool bIsLive;
@@ -70,24 +69,24 @@ private:
 	};
 
 	FSessionInfoInternal* AddSession(const FString& Path);
-	void RemoveSession(FSessionHandle Handle);
-	FSessionInfoInternal* GetSession(FSessionHandle Handle) const;
+	void RemoveSession(FStoreSessionHandle Handle);
+	FSessionInfoInternal* GetSession(FStoreSessionHandle Handle) const;
 	bool Tick(float DeltaTime);
 	void OnDirectoryChanged(const TArray<FFileChangeData>& FileChanges);
-	bool IsSessionLive(FSessionHandle Handle) const;
-	void CloseSessionStream(FSessionHandle Handle);
+	bool IsSessionLive(FStoreSessionHandle Handle) const;
+	void CloseSessionStream(FStoreSessionHandle Handle);
 
 	mutable FCriticalSection SessionsCS;
 	FString StoreDir;
 	IDirectoryWatcher* DirectoryWatcher;
 	FDelegateHandle DirectoryWatcherHandle;
 	FDelegateHandle TickHandle;
-	FSessionHandle NextSessionHandle = 1;
+	FStoreSessionHandle NextSessionHandle = 1;
 	TArray<FSessionInfoInternal*> Sessions;
 	TMap<FString, FSessionInfoInternal*> SessionsByPathMap;
 };
 
-FFileStoreOutDataStream::FFileStoreOutDataStream(FFileStore* InOwner, FSessionHandle InSessionHandle, IFileHandle* InFile)
+FFileStoreOutDataStream::FFileStoreOutDataStream(FFileStore* InOwner, FStoreSessionHandle InSessionHandle, IFileHandle* InFile)
 	: Owner(InOwner)
 	, SessionHandle(InSessionHandle)
 {
@@ -106,7 +105,7 @@ bool FFileStoreOutDataStream::Write(const void* Data, uint32 Size)
 	return Inner->Write((const uint8*)Data, Size);
 }
 
-FFileStoreInDataStream::FFileStoreInDataStream(FFileStore* InOwner, FSessionHandle InSessionHandle, FFileStream* InFileStream)
+FFileStoreInDataStream::FFileStoreInDataStream(FFileStore* InOwner, FStoreSessionHandle InSessionHandle, FFileStream* InFileStream)
 	: Owner(InOwner)
 	, SessionHandle(InSessionHandle)
 	, Inner(InFileStream)
@@ -192,13 +191,14 @@ FFileStore::FSessionInfoInternal* FFileStore::AddSession(const FString& Path)
 	Session->Name = FPaths::GetBaseFilename(Path);
 	Session->Path = Path;
 	Session->bIsValid = true;
+	Session->bIsLive = false;
 	Sessions.Add(Session);
 	SessionsByPathMap.Add(Path, Session);
 
 	return Session;
 }
 
-void FFileStore::RemoveSession(FSessionHandle Handle)
+void FFileStore::RemoveSession(FStoreSessionHandle Handle)
 {
 	FScopeLock Lock(&SessionsCS);
 
@@ -207,7 +207,7 @@ void FFileStore::RemoveSession(FSessionHandle Handle)
 	Sessions[Handle - 1]->bIsValid = false;
 }
 
-FFileStore::FSessionInfoInternal* FFileStore::GetSession(FSessionHandle Handle) const
+FFileStore::FSessionInfoInternal* FFileStore::GetSession(FStoreSessionHandle Handle) const
 {
 	FScopeLock Lock(&SessionsCS);
 	if (Handle - 1 >= Sessions.Num())
@@ -222,35 +222,24 @@ FFileStore::FSessionInfoInternal* FFileStore::GetSession(FSessionHandle Handle)
 	return Session;
 }
 
-void FFileStore::GetAvailableSessions(TArray<FSessionHandle>& OutTraces)
+void FFileStore::GetAvailableSessions(TArray<FStoreSessionInfo>& OutSessions) const
 {
 	FScopeLock Lock(&SessionsCS);
-	OutTraces.Empty(Sessions.Num());
+	OutSessions.Reserve(OutSessions.Num() + Sessions.Num());
 	for (const FSessionInfoInternal* Session : Sessions)
 	{
 		if (Session->bIsValid)
 		{
-			OutTraces.Add(Session->Handle);
+			FStoreSessionInfo& SessionInfo = OutSessions.AddDefaulted_GetRef();
+			SessionInfo.Handle = Session->Handle;
+			SessionInfo.Uri = *Session->Path;
+			SessionInfo.Name = *Session->Name;
+			SessionInfo.bIsLive = Session->bIsLive;
 		}
 	}
 }
 
-bool FFileStore::GetSessionInfo(FSessionHandle Handle, FSessionInfo& OutInfo)
-{
-	FScopeLock Lock(&SessionsCS);
-	
-	FSessionInfoInternal* Session = GetSession(Handle);
-	if (!Session)
-	{
-		return false;
-	}
-	OutInfo.Uri = *Session->Path;
-	OutInfo.Name = *Session->Name;
-	OutInfo.bIsLive = Session->bIsLive;
-	return true;
-}
-
-IOutDataStream* FFileStore::CreateNewSession()
+TTuple<FStoreSessionHandle, IOutDataStream*> FFileStore::CreateNewSession()
 {
 	FScopeLock Lock(&SessionsCS);
 
@@ -266,16 +255,15 @@ IOutDataStream* FFileStore::CreateNewSession()
 	IFileHandle* File = FileSystem.OpenWrite(*TracePath, true, true);
 	if (!File)
 	{
-		return nullptr;
+		return MakeTuple(FStoreSessionHandle(0), (IOutDataStream*)nullptr);
 	}
 
 	FSessionInfoInternal* Session = AddSession(TracePath);
 	Session->bIsLive = true;
-
-	return new FFileStoreOutDataStream(this, Session->Handle, File);
+	return MakeTuple(Session->Handle, static_cast<IOutDataStream*>(new FFileStoreOutDataStream(this, Session->Handle, File)));
 }
 
-IInDataStream* FFileStore::OpenSessionStream(FSessionHandle Handle)
+IInDataStream* FFileStore::OpenSessionStream(FStoreSessionHandle Handle)
 {
 	FScopeLock Lock(&SessionsCS);
 	
@@ -326,7 +314,7 @@ void FFileStore::OnDirectoryChanged(const TArray<FFileChangeData>& FileChanges)
 	}
 }
 
-bool FFileStore::IsSessionLive(FSessionHandle Handle) const
+bool FFileStore::IsSessionLive(FStoreSessionHandle Handle) const
 {
 	FScopeLock Lock(&SessionsCS);
 	
@@ -338,7 +326,7 @@ bool FFileStore::IsSessionLive(FSessionHandle Handle) const
 	return Session->bIsLive;
 }
 
-void FFileStore::CloseSessionStream(FSessionHandle Handle)
+void FFileStore::CloseSessionStream(FStoreSessionHandle Handle)
 {
 	FScopeLock Lock(&SessionsCS);
 	
diff --git a/Engine/Source/Developer/TraceAnalysis/Public/Trace/ControlClient.h b/Engine/Source/Developer/TraceAnalysis/Public/Trace/ControlClient.h
--- a/Engine/Source/Developer/TraceAnalysis/Public/Trace/ControlClient.h
+++ b/Engine/Source/Developer/TraceAnalysis/Public/Trace/ControlClient.h
@@ -5,6 +5,7 @@
 #include "CoreTypes.h"
 
 class FSocket;
+class FInternetAddr;
 
 namespace Trace
 {
@@ -15,6 +16,7 @@ class TRACEANALYSIS_API FControlClient
 public:
                 ~FControlClient();
     bool        Connect(const TCHAR* Host, uint16 Port=1985);
+	bool        Connect(const FInternetAddr& Address);
     void        Disconnect();
     bool        IsConnected() const;
     void        SendConnect(const TCHAR* Path);
diff --git a/Engine/Source/Developer/TraceAnalysis/Public/Trace/Recorder.h b/Engine/Source/Developer/TraceAnalysis/Public/Trace/Recorder.h
--- a/Engine/Source/Developer/TraceAnalysis/Public/Trace/Recorder.h
+++ b/Engine/Source/Developer/TraceAnalysis/Public/Trace/Recorder.h
@@ -8,6 +8,14 @@ namespace Trace
 {
 
 class IStore;
+typedef uint64 FStoreSessionHandle;
+typedef uint64 FRecorderSessionHandle;
+
+struct FRecorderSessionInfo
+{
+	FRecorderSessionHandle Handle;
+	FStoreSessionHandle StoreSessionHandle;
+};
 
 ////////////////////////////////////////////////////////////////////////////////
 class IRecorder
@@ -26,6 +34,12 @@ public:
 
 	/** Returns the number of current in-bound sessions. */
 	virtual uint32 GetSessionCount() const = 0;
+
+	/** Return information about the current in-bound sessions */
+	virtual void GetActiveSessions(TArray<FRecorderSessionInfo>& OutSessions) const = 0;
+
+	/** Toggles an event logger widlcard on or off for an active recording session */
+	virtual bool ToggleEvent(FRecorderSessionHandle RecordingHandle, const TCHAR* LoggerWildcard, bool bState) = 0;
 };
 
 TRACEANALYSIS_API TSharedPtr<IRecorder> Recorder_Create(TSharedRef<IStore> Store);
diff --git a/Engine/Source/Developer/TraceAnalysis/Public/Trace/Store.h b/Engine/Source/Developer/TraceAnalysis/Public/Trace/Store.h
--- a/Engine/Source/Developer/TraceAnalysis/Public/Trace/Store.h
+++ b/Engine/Source/Developer/TraceAnalysis/Public/Trace/Store.h
@@ -11,10 +11,11 @@ namespace Trace
 
 class IInDataStream;
 class IOutDataStream;
-typedef uint64 FSessionHandle;
+typedef uint64 FStoreSessionHandle;
 
-struct FSessionInfo
+struct FStoreSessionInfo
 {
+	FStoreSessionHandle Handle;
 	const TCHAR* Uri;
 	const TCHAR* Name;
 	bool bIsLive;
@@ -24,10 +25,9 @@ class IStore
 {
 public:
 	virtual ~IStore() = default;
-	virtual void GetAvailableSessions(TArray<FSessionHandle>& OutSessions) = 0;
-	virtual bool GetSessionInfo(FSessionHandle Handle, FSessionInfo& OutInfo) = 0;
-	virtual IOutDataStream* CreateNewSession() = 0;
-	virtual IInDataStream* OpenSessionStream(FSessionHandle Handle) = 0;
+	virtual void GetAvailableSessions(TArray<FStoreSessionInfo>& OutSessions) const = 0;
+	virtual TTuple<FStoreSessionHandle, IOutDataStream*> CreateNewSession() = 0;
+	virtual IInDataStream* OpenSessionStream(FStoreSessionHandle Handle) = 0;
 };
 
 TRACEANALYSIS_API TSharedPtr<IStore> Store_Create(const TCHAR* StoreDir);
diff --git a/Engine/Source/Developer/TraceServices/Private/SessionService.cpp b/Engine/Source/Developer/TraceServices/Private/SessionService.cpp
--- a/Engine/Source/Developer/TraceServices/Private/SessionService.cpp
+++ b/Engine/Source/Developer/TraceServices/Private/SessionService.cpp
@@ -2,10 +2,10 @@
 
 #include "TraceServices/SessionService.h"
 #include "SessionServicePrivate.h"
-#include "Trace/Store.h"
-#include "Trace/Recorder.h"
 #include "Trace/DataStream.h"
 #include "Misc/Paths.h"
+#include "Misc/OutputDeviceRedirector.h"
+#include "Trace/ControlClient.h"
 
 namespace Trace
 {
@@ -15,11 +15,18 @@ FSessionService::FSessionService()
 	LocalSessionDirectory = FPaths::ProjectSavedDir() / TEXT("TraceSessions");
 	TraceStore = Store_Create(*LocalSessionDirectory);
 	TraceRecorder = Recorder_Create(TraceStore.ToSharedRef());
+
+	FTickerDelegate TickDelegate = FTickerDelegate::CreateRaw(this, &FSessionService::Tick);
+	TickHandle = FTicker::GetCoreTicker().AddTicker(TickDelegate, 0.5f);
 }
 
 FSessionService::~FSessionService()
 {
-	
+	if (TickHandle.IsValid())
+	{
+		FTicker::GetCoreTicker().RemoveTicker(TickHandle);
+		TickHandle.Reset();
+	}
 }
 
 bool FSessionService::StartRecorderServer()
@@ -37,17 +44,44 @@ void FSessionService::StopRecorderServer()
 	TraceRecorder->StopRecording();
 }
 
-void FSessionService::GetAvailableSessions(TArray<Trace::FSessionHandle>& OutSessions)
+void FSessionService::GetAvailableSessions(TArray<FSessionHandle>& OutSessions) const
+{
+	FScopeLock Lock(&SessionsCS);
+	OutSessions.Reserve(OutSessions.Num() + Sessions.Num());
+	for (const auto& KV : Sessions)
+	{
+		OutSessions.Add(static_cast<FSessionHandle>(KV.Key));
+	}
+}
+
+void FSessionService::GetLiveSessions(TArray<FSessionHandle>& OutSessions) const
 {
-	TraceStore->GetAvailableSessions(OutSessions);
+	FScopeLock Lock(&SessionsCS);
+	OutSessions.Reserve(OutSessions.Num() + Sessions.Num());
+	for (const auto& KV : Sessions)
+	{
+		if (KV.Value.bIsLive)
+		{
+			OutSessions.Add(static_cast<FSessionHandle>(KV.Key));
+		}
+	}
 }
 
-bool FSessionService::GetSessionInfo(Trace::FSessionHandle SessionHandle, Trace::FSessionInfo& OutSessionInfo)
+bool FSessionService::GetSessionInfo(FSessionHandle SessionHandle, FSessionInfo& OutSessionInfo) const
 {
-	return TraceStore->GetSessionInfo(SessionHandle, OutSessionInfo);
+	FScopeLock Lock(&SessionsCS);
+	const FSessionInfoInternal* FindIt = Sessions.Find(SessionHandle);
+	if (!FindIt)
+	{
+		return false;
+	}
+	OutSessionInfo.Uri = FindIt->Uri;
+	OutSessionInfo.Name = FindIt->Name;
+	OutSessionInfo.bIsLive = FindIt->bIsLive;
+	return true;
 }
 
-Trace::IInDataStream* FSessionService::OpenSessionStream(Trace::FSessionHandle SessionHandle)
+Trace::IInDataStream* FSessionService::OpenSessionStream(FSessionHandle SessionHandle)
 {
 	return TraceStore->OpenSessionStream(SessionHandle);
 }
@@ -57,4 +91,37 @@ Trace::IInDataStream* FSessionService::OpenSessionFromFile(const TCHAR* FilePath
 	return Trace::DataStream_ReadFile(FilePath);
 }
 
+bool FSessionService::Tick(float DeltaTime)
+{
+	UpdateSessions();
+	return true;
+}
+
+void FSessionService::UpdateSessions()
+{
+	TArray<FStoreSessionInfo> StoreSessions;
+	TraceStore->GetAvailableSessions(StoreSessions);
+	TArray<FRecorderSessionInfo> RecorderSessions;
+	TraceRecorder->GetActiveSessions(RecorderSessions);
+
+	FScopeLock Lock(&SessionsCS);
+	for (const FStoreSessionInfo& StoreSession : StoreSessions)
+	{
+		FSessionInfoInternal& Session = Sessions.FindOrAdd(StoreSession.Handle);
+		Session.Uri = StoreSession.Uri;
+		Session.Name = StoreSession.Name;
+		Session.bIsLive = StoreSession.bIsLive;
+		Session.RecorderSessionHandle = 0;
+	}
+
+	for (const FRecorderSessionInfo& RecorderSession : RecorderSessions)
+	{
+		FSessionInfoInternal* FindIt = Sessions.Find(RecorderSession.StoreSessionHandle);
+		if (FindIt)
+		{
+			FindIt->RecorderSessionHandle = RecorderSession.Handle;
+		}
+	}
+}
+
 }
diff --git a/Engine/Source/Developer/TraceServices/Private/SessionServicePrivate.h b/Engine/Source/Developer/TraceServices/Private/SessionServicePrivate.h
--- a/Engine/Source/Developer/TraceServices/Private/SessionServicePrivate.h
+++ b/Engine/Source/Developer/TraceServices/Private/SessionServicePrivate.h
@@ -3,12 +3,14 @@
 #pragma once
 
 #include "TraceServices/SessionService.h"
+#include "Misc/ScopeLock.h"
+#include "Containers/Ticker.h"
+#include "Trace/Recorder.h"
+#include "Trace/Store.h"
 
 namespace Trace
 {
 
-class IRecorder;
-
 class FSessionService
 	: public ISessionService
 {
@@ -19,15 +21,30 @@ public:
 	virtual bool IsRecorderServerRunning() const override;
 	virtual void StopRecorderServer() override;
 	virtual const TCHAR* GetLocalSessionDirectory() const override { return *LocalSessionDirectory; }
-	virtual void GetAvailableSessions(TArray<Trace::FSessionHandle>& OutSessions) override;
-	virtual bool GetSessionInfo(Trace::FSessionHandle SessionHandle, Trace::FSessionInfo& OutSessionInfo) override;
+	virtual void GetAvailableSessions(TArray<Trace::FSessionHandle>& OutSessions) const override;
+	virtual void GetLiveSessions(TArray<Trace::FSessionHandle>& OutSessions) const override;
+	virtual bool GetSessionInfo(Trace::FSessionHandle SessionHandle, Trace::FSessionInfo& OutSessionInfo) const override;
 	virtual Trace::IInDataStream* OpenSessionStream(Trace::FSessionHandle SessionHandle) override;
 	virtual Trace::IInDataStream* OpenSessionFromFile(const TCHAR* FilePath) override;
 
 private:
+	bool Tick(float DeltaTime);
+	void UpdateSessions();
+
+	struct FSessionInfoInternal
+	{
+		FRecorderSessionHandle RecorderSessionHandle;
+		const TCHAR* Uri;
+		const TCHAR* Name;
+		bool bIsLive;
+	};
+
 	FString LocalSessionDirectory;
 	TSharedPtr<IStore> TraceStore;
 	TSharedPtr<IRecorder> TraceRecorder;
+	mutable FCriticalSection SessionsCS;
+	TMap<FStoreSessionHandle, FSessionInfoInternal> Sessions;
+	FDelegateHandle TickHandle;
 };
 
 }
diff --git a/Engine/Source/Developer/TraceServices/Public/TraceServices/SessionService.h b/Engine/Source/Developer/TraceServices/Public/TraceServices/SessionService.h
--- a/Engine/Source/Developer/TraceServices/Public/TraceServices/SessionService.h
+++ b/Engine/Source/Developer/TraceServices/Public/TraceServices/SessionService.h
@@ -10,6 +10,15 @@
 namespace Trace
 {
 
+typedef uint64 FSessionHandle;
+
+struct FSessionInfo
+{
+	const TCHAR* Uri;
+	const TCHAR* Name;
+	bool bIsLive;
+};
+
 class ISessionService
 {
 public:
@@ -18,9 +27,10 @@ public:
 	virtual bool IsRecorderServerRunning() const = 0;
 	virtual void StopRecorderServer() = 0;
 	virtual const TCHAR* GetLocalSessionDirectory() const = 0;
-	virtual void GetAvailableSessions(TArray<Trace::FSessionHandle>& OutSessions) = 0;
-	virtual bool GetSessionInfo(Trace::FSessionHandle SessionHandle, Trace::FSessionInfo& OutSessionInfo) = 0;
-	virtual Trace::IInDataStream* OpenSessionStream(Trace::FSessionHandle SessionHandle) = 0;
+	virtual void GetAvailableSessions(TArray<FSessionHandle>& OutSessions) const = 0;
+	virtual void GetLiveSessions(TArray<FSessionHandle>& OutSessions) const = 0;
+	virtual bool GetSessionInfo(FSessionHandle SessionHandle, FSessionInfo& OutSessionInfo) const = 0;
+	virtual Trace::IInDataStream* OpenSessionStream(FSessionHandle SessionHandle) = 0;
 	virtual Trace::IInDataStream* OpenSessionFromFile(const TCHAR* FilePath) = 0;
 };
 
