Removed used of FPlatformMemory::Memory*() methods.

diff --git a/Engine/Source/Developer/TraceAnalysis/Private/DataStream.cpp b/Engine/Source/Developer/TraceAnalysis/Private/DataStream.cpp
--- a/Engine/Source/Developer/TraceAnalysis/Private/DataStream.cpp
+++ b/Engine/Source/Developer/TraceAnalysis/Private/DataStream.cpp
@@ -64,12 +64,12 @@ void FFileStream::OpenFileInternal()
 FStreamReader::FStreamReader(IInDataStream& InDataStream)
 	: DataStream(InDataStream)
 {
-	Buffer = (uint8*)FPlatformMemory::MemoryRangeReserve(BufferSize, true);
+	Buffer = new uint8[BufferSize];
 }
 
 FStreamReader::~FStreamReader()
 {
-	FPlatformMemory::MemoryRangeFree(Buffer, BufferSize);
+	delete[] Buffer;
 }
 
 IInDataStream* DataStream_ReadFile(const TCHAR* FilePath)
