Fixed compile errors (missing #include "Misc/ScopeLock.h") - NonUnity / win64

diff --git a/Engine/Source/Developer/TraceAnalysis/Private/Recorder/Recorder.cpp b/Engine/Source/Developer/TraceAnalysis/Private/Recorder/Recorder.cpp
--- a/Engine/Source/Developer/TraceAnalysis/Private/Recorder/Recorder.cpp
+++ b/Engine/Source/Developer/TraceAnalysis/Private/Recorder/Recorder.cpp
@@ -6,6 +6,7 @@
 #include "HAL/Runnable.h"
 #include "HAL/RunnableThread.h"
 #include "IPAddress.h"
+#include "Misc/ScopeLock.h"
 #include "Sockets.h"
 #include "SocketSubsystem.h"
 #include "Trace/DataStream.h"
diff --git a/Engine/Source/Runtime/Core/Public/HAL/MallocBinnedArena.h b/Engine/Source/Runtime/Core/Public/HAL/MallocBinnedArena.h
--- a/Engine/Source/Runtime/Core/Public/HAL/MallocBinnedArena.h
+++ b/Engine/Source/Runtime/Core/Public/HAL/MallocBinnedArena.h
@@ -6,6 +6,7 @@
 
 #if PLATFORM_64BITS && PLATFORM_HAS_FPlatformVirtualMemoryBlock
 #include "Misc/AssertionMacros.h"
+#include "Misc/ScopeLock.h"
 #include "HAL/MemoryBase.h"
 #include "HAL/UnrealMemory.h"
 #include "Math/NumericLimits.h"
diff --git a/Engine/Source/Runtime/Core/Public/HAL/MallocBinnedGPU.h b/Engine/Source/Runtime/Core/Public/HAL/MallocBinnedGPU.h
--- a/Engine/Source/Runtime/Core/Public/HAL/MallocBinnedGPU.h
+++ b/Engine/Source/Runtime/Core/Public/HAL/MallocBinnedGPU.h
@@ -6,6 +6,7 @@
 
 #if PLATFORM_64BITS && PLATFORM_HAS_FPlatformVirtualMemoryBlock
 #include "Misc/AssertionMacros.h"
+#include "Misc/ScopeLock.h"
 #include "HAL/MemoryBase.h"
 #include "HAL/UnrealMemory.h"
 #include "Math/NumericLimits.h"
