Fixed and finished the documentation for trace's IRecorder.

diff --git a/Engine/Source/Developer/TraceAnalysis/Public/Trace/Recorder.h b/Engine/Source/Developer/TraceAnalysis/Public/Trace/Recorder.h
--- a/Engine/Source/Developer/TraceAnalysis/Public/Trace/Recorder.h
+++ b/Engine/Source/Developer/TraceAnalysis/Public/Trace/Recorder.h
@@ -11,13 +11,21 @@ class IStore;
 typedef uint64 FStoreSessionHandle;
 typedef uint64 FRecorderSessionHandle;
 
+/**
+ * Information about traces being recorded
+ */
 struct FRecorderSessionInfo
 {
+	/** Identifier of the recording session */
 	FRecorderSessionHandle Handle;
+	/** Store identifier of the trace in question */
 	FStoreSessionHandle StoreSessionHandle;
 };
 
-////////////////////////////////////////////////////////////////////////////////
+/**
+ * Interface to an object that can receive traces from remote runtimes and write
+ * them into a store.
+ */
 class IRecorder
 {
 public:
@@ -35,13 +43,22 @@ public:
 	/** Returns the number of current in-bound sessions. */
 	virtual uint32 GetSessionCount() const = 0;
 
-	/** Return information about the current in-bound sessions */
+	/** Return information about the current in-bound sessions being recorded
+	 * @param OutSessions TArray instance to receive the current session info. */
 	virtual void GetActiveSessions(TArray<FRecorderSessionInfo>& OutSessions) const = 0;
 
-	/** Toggles an event logger widlcard on or off for an active recording session */
-	virtual bool ToggleEvent(FRecorderSessionHandle RecordingHandle, const TCHAR* LoggerWildcard, bool bState) = 0;
+	/** Toggles an event logger wildcard on or off for an active recording session.
+	 * @param RecordingHandle Identifier of the active session, from GetActiveSessions().
+	 * @param EventSpec Event specification (see FControlClient::ToggleEvent)
+	 * @param bState True to enable the event and false to disable it.
+	 * @return True if the command was successfully sent. */
+	virtual bool ToggleEvent(FRecorderSessionHandle RecordingHandle, const TCHAR* EventSpec, bool bState) = 0;
 };
 
+/**
+ * Creates a recorder that will allow runtimes to connect to it and send their trace events.
+ * @param Store The store instance to save trace streams into.
+ */
 TRACEANALYSIS_API TSharedPtr<IRecorder> Recorder_Create(TSharedRef<IStore> Store);
 
 } // namespace Trace
