Documented trace's ControlClient

diff --git a/Engine/Source/Developer/TraceAnalysis/Private/ControlClient.cpp b/Engine/Source/Developer/TraceAnalysis/Private/ControlClient.cpp
--- a/Engine/Source/Developer/TraceAnalysis/Private/ControlClient.cpp
+++ b/Engine/Source/Developer/TraceAnalysis/Private/ControlClient.cpp
@@ -101,9 +101,9 @@ void FControlClient::SendConnect(const TCHAR* Path)
 }
 
 ////////////////////////////////////////////////////////////////////////////////
-void FControlClient::SendToggleEvent(const TCHAR* EventMask, bool bState)
+void FControlClient::SendToggleEvent(const TCHAR* EventSpec, bool bState)
 {
-    FormatAndSend(TEXT("ToggleEvent %s %d"), EventMask, bState);
+    FormatAndSend(TEXT("ToggleEvent %s %d"), EventSpec, bState);
 }
 
 ////////////////////////////////////////////////////////////////////////////////
diff --git a/Engine/Source/Developer/TraceAnalysis/Public/Trace/ControlClient.h b/Engine/Source/Developer/TraceAnalysis/Public/Trace/ControlClient.h
--- a/Engine/Source/Developer/TraceAnalysis/Public/Trace/ControlClient.h
+++ b/Engine/Source/Developer/TraceAnalysis/Public/Trace/ControlClient.h
@@ -10,18 +10,46 @@ class FInternetAddr;
 namespace Trace
 {
 
-////////////////////////////////////////////////////////////////////////////////
+/**
+ * Simple helper class for communicating with trace system of a running instance.
+ */
 class TRACEANALYSIS_API FControlClient
 {
 public:
-                ~FControlClient();
-    bool        Connect(const TCHAR* Host, uint16 Port=1985);
-	bool        Connect(const FInternetAddr& Address);
-    void        Disconnect();
-    bool        IsConnected() const;
-    void        SendConnect(const TCHAR* Path);
-    void        SendToggleEvent(const TCHAR* Logger, bool bState=true);
-    void        Send(const TCHAR* Command);
+	/** The destructor will ensure disconnection the FControlClient class
+	 * goes out of scope. */
+    ~FControlClient();
+
+	/** Initiates a connection to a runtime instance to control.
+	 * @param Host IP address or host name of where the runtime is to be found.
+	 * @returns True if a connection was made successfully. */
+    bool Connect(const TCHAR* Host, uint16 Port=1985);
+
+	/** Initiates a connection to a runtime instance to control.
+	 * @param Fully qualified FInternetAddr object of the runtime's location.
+	 * @returns True if a connection was made successfully. */
+	bool Connect(const FInternetAddr& Address);
+
+	/** Disconnects the client from the runtime. */
+    void Disconnect();
+
+	/** Returns true if the client is currently connected to a runtime */
+    bool IsConnected() const;
+
+	/** Tells the runtime to attempt to record its trace event log.
+	 * @param Path Path that the runtime should send events to. */
+    void SendConnect(const TCHAR* Path);
+
+	/** Toggles events on and off. Events can be specified in two forms; either
+	 * by logger to address a logger's events as a whole (EventSpec="LoggerName")
+	 * or individually (EventSpec="LoggerName.EventName").
+	 * @param EventSpec Specification of the event to toggle.
+	 * @param bState True to enable the event and false to disable it. */
+    void SendToggleEvent(const TCHAR* EventSpec, bool bState=true);
+
+	/** Sends a raw command to the instance.
+	 * @param Command The command to send to the runtime (CRLF terminator not required). */
+    void Send(const TCHAR* Command);
 
 private:
     void        FormatAndSend(const TCHAR* Format, ...);
