Documented trace's IStore.

diff --git a/Engine/Source/Developer/TraceAnalysis/Public/Trace/Store.h b/Engine/Source/Developer/TraceAnalysis/Public/Trace/Store.h
--- a/Engine/Source/Developer/TraceAnalysis/Public/Trace/Store.h
+++ b/Engine/Source/Developer/TraceAnalysis/Public/Trace/Store.h
@@ -13,23 +13,47 @@ class IInDataStream;
 class IOutDataStream;
 typedef uint64 FStoreSessionHandle;
 
+/**
+ * Information about traces available in the trace store
+ */
 struct FStoreSessionInfo
 {
+	/** A handle to identify the trace */
 	FStoreSessionHandle Handle;
+	/** Fully qualified path to the trace store item */
 	const TCHAR* Uri;
+	/** Friendly-ish name of the trace */
 	const TCHAR* Name;
+	/** True if the trace is currently being written to */
 	bool bIsLive;
 };
 
+/**
+ * Interface to a provider of storage for traces.
+ */
 class IStore
 {
 public:
 	virtual ~IStore() = default;
+
+	/** Enumerates the "sessions" (aka traces) that are currently available.
+	 * @param OutSessions Recipient of the list of available traces. */
 	virtual void GetAvailableSessions(TArray<FStoreSessionInfo>& OutSessions) const = 0;
+
+	/** Creates a new trace that can be written to.
+	 * @return A handle for identifying the "session" (aka trace) and an interface for writing arbitrary data to the trace */
 	virtual TTuple<FStoreSessionHandle, IOutDataStream*> CreateNewSession() = 0;
+
+	/** Opens an existing trace.
+	 * @param Handle Identifier for the trace to open
+	 * @return Interface to reading the arbitrary trace data */
 	virtual IInDataStream* OpenSessionStream(FStoreSessionHandle Handle) = 0;
 };
 
+/**
+ * Creates a file-based store at a given file system location location
+ * @param StoreDir File system directory where the store will read/write traces
+ */
 TRACEANALYSIS_API TSharedPtr<IStore> Store_Create(const TCHAR* StoreDir);
 
 }
