Documented the objects involved in describing and running the analysis of a trace.

diff --git a/Engine/Source/Developer/TraceAnalysis/Private/Analysis/Context.cpp b/Engine/Source/Developer/TraceAnalysis/Private/Analysis/Context.cpp
--- a/Engine/Source/Developer/TraceAnalysis/Private/Analysis/Context.cpp
+++ b/Engine/Source/Developer/TraceAnalysis/Private/Analysis/Context.cpp
@@ -1,6 +1,6 @@
 // Copyright 1998-2019 Epic Games, Inc. All Rights Reserved.
 
-#include "Engine.h"
+#include "Processor.h"
 #include "Trace/Analysis.h"
 
 namespace Trace
@@ -13,15 +13,13 @@ void FAnalysisContext::AddAnalyzer(IAnalyzer& Analyzer)
 }
 
 ////////////////////////////////////////////////////////////////////////////////
-FAnalysisProcessor FAnalysisContext::Process()
+FAnalysisProcessor FAnalysisContext::Process(IInDataStream& DataStream)
 {
 	FAnalysisProcessor Processor;
-
-	if (Analyzers.Num())
+	if (Analyzers.Num() > 0)
 	{
-		Processor.Impl = new FAnalysisEngine(MoveTemp(Analyzers));
+		Processor.Impl = new FAnalysisProcessor::FImpl(DataStream, MoveTemp(Analyzers));
 	}
-
 	return MoveTemp(Processor);
 }
 
diff --git a/Engine/Source/Developer/TraceAnalysis/Private/Analysis/Processor.cpp b/Engine/Source/Developer/TraceAnalysis/Private/Analysis/Processor.cpp
--- a/Engine/Source/Developer/TraceAnalysis/Private/Analysis/Processor.cpp
+++ b/Engine/Source/Developer/TraceAnalysis/Private/Analysis/Processor.cpp
@@ -1,30 +1,97 @@
 // Copyright 1998-2019 Epic Games, Inc. All Rights Reserved.
 
-#include "Engine.h"
+#include "Processor.h"
+#include "HAL/Event.h"
+#include "HAL/PlatformProcess.h"
+#include "HAL/RunnableThread.h"
+#include "Templates/UnrealTemplate.h"
 #include "Trace/Analysis.h"
 
 namespace Trace
 {
 
 ////////////////////////////////////////////////////////////////////////////////
-void FAnalysisProcessor::Start(IInDataStream& DataStream)
+FAnalysisProcessor::FImpl::FImpl(IInDataStream& InDataStream, TArray<IAnalyzer*>&& InAnalyzers)
+: AnalysisEngine(Forward<TArray<IAnalyzer*>>(InAnalyzers))
+, DataStream(InDataStream)
+, StopEvent(FPlatformProcess::GetSynchEventFromPool(true))
+, UnpausedEvent(FPlatformProcess::GetSynchEventFromPool(true))
 {
-	FAnalysisEngine* Engine = (FAnalysisEngine*)Impl;
-	if (Engine == nullptr)
-	{
-		return;
-	}
+	Thread = FRunnableThread::Create(this, TEXT("TraceAnalysis"));
+	PauseAnalysis(false);
+}
 
-	FStreamReader Reader(DataStream);
-	while (FStreamReader::FData* Data = Reader.Read())
+////////////////////////////////////////////////////////////////////////////////
+FAnalysisProcessor::FImpl::~FImpl()
+{
+	StopAnalysis();
+	FPlatformProcess::ReturnSynchEventToPool(UnpausedEvent);
+	FPlatformProcess::ReturnSynchEventToPool(StopEvent);
+}
+
+////////////////////////////////////////////////////////////////////////////////
+uint32 FAnalysisProcessor::FImpl::Run()
+{
+	for (FStreamReader Reader(DataStream); FStreamReader::FData* Data = Reader.Read();)
 	{
-		if (!Engine->OnData(*Data))
+		if (StopEvent->Wait(0, true))
+		{
+			break;
+		}
+
+		UnpausedEvent->Wait();
+
+		if (!AnalysisEngine.OnData(*Data))
 		{
 			break;
 		}
 	}
-	delete Engine;
-	Engine = nullptr;
+
+	return 0;
+}
+
+////////////////////////////////////////////////////////////////////////////////
+bool FAnalysisProcessor::FImpl::IsActive() const
+{
+	return (Thread != nullptr);
 }
 
+////////////////////////////////////////////////////////////////////////////////
+void FAnalysisProcessor::FImpl::StopAnalysis()
+{
+	if (IsActive())
+	{
+		StopEvent->Trigger();
+		WaitOnAnalysis();
+	}
+}
+
+////////////////////////////////////////////////////////////////////////////////
+void FAnalysisProcessor::FImpl::WaitOnAnalysis()
+{
+	if (IsActive())
+	{
+		Thread->Kill(true);
+		delete Thread;
+		Thread = nullptr;
+	}
+}
+
+////////////////////////////////////////////////////////////////////////////////
+void FAnalysisProcessor::FImpl::PauseAnalysis(bool bState)
+{
+	if (IsActive())
+	{
+		bState ? UnpausedEvent->Reset() : UnpausedEvent->Trigger();
+	}
+}
+
+
+
+////////////////////////////////////////////////////////////////////////////////
+bool FAnalysisProcessor::IsActive() const	{ return (Impl != nullptr) ? Impl->IsActive() : false; }
+void FAnalysisProcessor::Stop()				{ if (Impl != nullptr) { Impl->StopAnalysis(); } }
+void FAnalysisProcessor::Wait()				{ if (Impl != nullptr) { Impl->WaitOnAnalysis(); } }
+void FAnalysisProcessor::Pause(bool bState) { if (Impl != nullptr) { Impl->PauseAnalysis(bState); } }
+
 } // namespace Trace
diff --git a/Engine/Source/Developer/TraceAnalysis/Private/Analysis/Processor.h b/Engine/Source/Developer/TraceAnalysis/Private/Analysis/Processor.h
new file mode 100644
--- /dev/null
+++ b/Engine/Source/Developer/TraceAnalysis/Private/Analysis/Processor.h
@@ -0,0 +1,35 @@
+// Copyright 1998-2019 Epic Games, Inc. All Rights Reserved.
+
+#pragma once
+
+#include "Engine.h"
+#include "HAL/Runnable.h"
+
+class FEvent;
+class FRunnableThread;
+
+namespace Trace
+{
+
+////////////////////////////////////////////////////////////////////////////////
+class FAnalysisProcessor::FImpl
+	: public FRunnable
+{
+public:
+						FImpl(IInDataStream& DataStream, TArray<IAnalyzer*>&& InAnalyzers);
+						~FImpl();
+	virtual uint32		Run() override;
+	bool				IsActive() const;
+	void				StopAnalysis();
+	void				WaitOnAnalysis();
+	void				PauseAnalysis(bool bState);
+
+private:
+	FAnalysisEngine		AnalysisEngine;
+	IInDataStream&		DataStream;
+	FEvent*				StopEvent;
+	FEvent*				UnpausedEvent;
+	FRunnableThread*	Thread = nullptr;
+};
+
+} // namespace Trace
diff --git a/Engine/Source/Developer/TraceAnalysis/Public/Trace/Analysis.h b/Engine/Source/Developer/TraceAnalysis/Public/Trace/Analysis.h
--- a/Engine/Source/Developer/TraceAnalysis/Public/Trace/Analysis.h
+++ b/Engine/Source/Developer/TraceAnalysis/Public/Trace/Analysis.h
@@ -10,29 +10,51 @@ namespace Trace
 class IAnalyzer;
 class IInDataStream;
 
-////////////////////////////////////////////////////////////////////////////////
+/**
+ * Represents the processing (e.g. analysis) of a trace stream. Instances are
+ * created by constructing an FAnalysisContext object to marry an event trace
+ * with how it should be analyzed. Note that the processing (and thus analysis)
+ * happens on another thread.
+ */
 class TRACEANALYSIS_API FAnalysisProcessor
 {
 public:
-	explicit	operator bool () const;
-	void		Start(IInDataStream& DataStream);
-	void		Stop();
-	void		Wait();
-	void		Pause(bool State);
+	/** Checks if this object instance is valid and currently processing */
+	bool IsActive() const;
+
+	/** End processing a trace stream. */
+	void Stop();
+
+	/** Wait for the entire stream to have been processed and analysed. */
+	void Wait();
+
+	/** Pause or resume the processing.
+	 * @param bState Pause if true, resume if false. */
+	void Pause(bool bState);
 
 private:
 	friend		class FAnalysisContext;
-	void*		Impl = nullptr;
+	class		FImpl;
+	FImpl*		Impl = nullptr;
 };
 
 
 
-////////////////////////////////////////////////////////////////////////////////
+/**
+ * Used to describe how a log of trace events should be analyzed and being the
+ * analysis on a particular trace stream.  
+ */
 class TRACEANALYSIS_API FAnalysisContext
 {
 public:
-	void				AddAnalyzer(IAnalyzer& Analyzer);
-	FAnalysisProcessor	Process();
+	/** Adds an analyzer instance that will subscribe to and receive event data
+	 * from the trace stream. */
+	void AddAnalyzer(IAnalyzer& Analyzer);
+
+	/** Creates and starts analysis returning an FAnalysisProcessor instance which
+	 * represents the analysis and affords some control over it.
+	 * @param DataStream Input stream of trace log data to be analysed. */
+	FAnalysisProcessor Process(IInDataStream& DataStream);
 
 private:
 	TArray<IAnalyzer*>	Analyzers;
diff --git a/Engine/Source/Developer/TraceServices/Private/AnalysisService.cpp b/Engine/Source/Developer/TraceServices/Private/AnalysisService.cpp
--- a/Engine/Source/Developer/TraceServices/Private/AnalysisService.cpp
+++ b/Engine/Source/Developer/TraceServices/Private/AnalysisService.cpp
@@ -136,9 +136,8 @@ void FAnalysisService::AnalyzeInternal(TSharedRef<FAnalysisSession> AnalysisSess
 	{
 		Context.AddAnalyzer(*Analyzer);
 	}
-	Trace::FAnalysisProcessor Processor = Context.Process();
-
-	Processor.Start(*DataStream);
+	Trace::FAnalysisProcessor Processor = Context.Process(*DataStream);
+	Processor.Wait();
 
 	for (Trace::IAnalyzer* Analyzer : Analyzers)
 	{
