Removed IAnalyzer.GetArray() as array support isn't quite implemented yet.

diff --git a/Engine/Source/Developer/TraceAnalysis/Private/Analysis/Engine.cpp b/Engine/Source/Developer/TraceAnalysis/Private/Analysis/Engine.cpp
--- a/Engine/Source/Developer/TraceAnalysis/Private/Analysis/Engine.cpp
+++ b/Engine/Source/Developer/TraceAnalysis/Private/Analysis/Engine.cpp
@@ -218,7 +218,6 @@ struct FAnalysisEngine::FEventDataImpl
 {
 	virtual					~FEventDataImpl() = default;
 	virtual const FValue&	GetValue(const ANSICHAR* FieldName) const override;
-	virtual const FArray&	GetArray(const ANSICHAR* FieldName) const override;
 	virtual const uint8*	GetData() const override;
 	virtual const uint8*	GetAttachment() const override;
 	virtual uint16			GetAttachmentSize() const override;
@@ -251,12 +250,6 @@ const IAnalyzer::FValue& FAnalysisEngine::FEventDataImpl::GetValue(const ANSICHA
 	return *(FValue*)Ret;
 }
 
-////////////////////////////////////////////////////////////////////////////////
-const IAnalyzer::FArray& FAnalysisEngine::FEventDataImpl::GetArray(const ANSICHAR* FieldName) const
-{
-	return *(FArray*)0x493;
-}
-
 ////////////////////////////////////////////////////////////////////////////////
 const uint8* FAnalysisEngine::FEventDataImpl::GetData() const
 {
diff --git a/Engine/Source/Developer/TraceAnalysis/Public/Trace/Analyzer.h b/Engine/Source/Developer/TraceAnalysis/Public/Trace/Analyzer.h
--- a/Engine/Source/Developer/TraceAnalysis/Public/Trace/Analyzer.h
+++ b/Engine/Source/Developer/TraceAnalysis/Public/Trace/Analyzer.h
@@ -48,14 +48,9 @@ public:
 		AsType As() const;
 	};
 
-	struct FArray
-	{
-	};
-
 	struct FEventData
 	{
 		virtual const FValue&	GetValue(const ANSICHAR* FieldName) const = 0;
-		virtual const FArray&	GetArray(const ANSICHAR* FieldName) const = 0;
 		virtual const uint8*	GetData() const = 0;
 		virtual const uint8*	GetAttachment() const = 0;
 		virtual uint16			GetAttachmentSize() const = 0;
