Removed IAnalyzer.Event.GetData() and GetTotalSize() from public view as they really aren't useful and we don't want to leak internals so easily.

diff --git a/Engine/Source/Developer/TraceAnalysis/Private/Analysis/Engine.cpp b/Engine/Source/Developer/TraceAnalysis/Private/Analysis/Engine.cpp
--- a/Engine/Source/Developer/TraceAnalysis/Private/Analysis/Engine.cpp
+++ b/Engine/Source/Developer/TraceAnalysis/Private/Analysis/Engine.cpp
@@ -218,10 +218,8 @@ struct FAnalysisEngine::FEventDataImpl
 {
 	virtual					~FEventDataImpl() = default;
 	virtual const FValue&	GetValue(const ANSICHAR* FieldName) const override;
-	virtual const uint8*	GetData() const override;
 	virtual const uint8*	GetAttachment() const override;
 	virtual uint16			GetAttachmentSize() const override;
-	virtual uint16			GetTotalSize() const override;
 	const FDispatch*		Dispatch;
 	const uint8*			Ptr;
 	uint16					Size;
@@ -250,12 +248,6 @@ const IAnalyzer::FValue& FAnalysisEngine::FEventDataImpl::GetValue(const ANSICHA
 	return *(FValue*)Ret;
 }
 
-////////////////////////////////////////////////////////////////////////////////
-const uint8* FAnalysisEngine::FEventDataImpl::GetData() const
-{
-	return Ptr;
-}
-
 ////////////////////////////////////////////////////////////////////////////////
 const uint8* FAnalysisEngine::FEventDataImpl::GetAttachment() const
 {
@@ -268,12 +260,6 @@ uint16 FAnalysisEngine::FEventDataImpl::GetAttachmentSize() const
 	return Size - Dispatch->EventSize;
 }
 
-////////////////////////////////////////////////////////////////////////////////
-uint16 FAnalysisEngine::FEventDataImpl::GetTotalSize() const
-{
-	return Size;
-}
-
 
 
 ////////////////////////////////////////////////////////////////////////////////
@@ -458,7 +444,8 @@ FAnalysisEngine::FDispatch& FAnalysisEngine::AddDispatch(uint16 Uid, uint16 Fiel
 ////////////////////////////////////////////////////////////////////////////////
 void FAnalysisEngine::OnNewEvent(const FOnEventContext& Context)
 {
-	const auto& NewEvent = *(FNewEventEvent*)Context.EventData.GetData();
+	const FEventDataImpl& EventData = (const FEventDataImpl&)(Context.EventData);
+	const auto& NewEvent = *(FNewEventEvent*)(EventData.Ptr);
 
 	FDispatch& Dispatch = AddDispatch(NewEvent.EventUid, NewEvent.FieldCount);
 
diff --git a/Engine/Source/Developer/TraceAnalysis/Public/Trace/Analyzer.h b/Engine/Source/Developer/TraceAnalysis/Public/Trace/Analyzer.h
--- a/Engine/Source/Developer/TraceAnalysis/Public/Trace/Analyzer.h
+++ b/Engine/Source/Developer/TraceAnalysis/Public/Trace/Analyzer.h
@@ -51,10 +51,8 @@ public:
 	struct FEventData
 	{
 		virtual const FValue&	GetValue(const ANSICHAR* FieldName) const = 0;
-		virtual const uint8*	GetData() const = 0;
 		virtual const uint8*	GetAttachment() const = 0;
 		virtual uint16			GetAttachmentSize() const = 0;
-		virtual uint16			GetTotalSize() const = 0;
 	};
 
 	struct FOnEventContext
diff --git a/Engine/Source/Developer/TraceServices/Private/Analyzers/CpuProfilerTraceAnalysis.cpp b/Engine/Source/Developer/TraceServices/Private/Analyzers/CpuProfilerTraceAnalysis.cpp
--- a/Engine/Source/Developer/TraceServices/Private/Analyzers/CpuProfilerTraceAnalysis.cpp
+++ b/Engine/Source/Developer/TraceServices/Private/Analyzers/CpuProfilerTraceAnalysis.cpp
@@ -47,7 +47,7 @@ void FCpuProfilerAnalyzer::OnEvent(uint16 RouteId, const FOnEventContext& Contex
 	case RouteId_EventBatch:
 	case RouteId_EndCapture:
 	{
-		TotalEventSize += EventData.GetTotalSize();
+		TotalEventSize += EventData.GetAttachmentSize();
 		uint32 ThreadId = EventData.GetValue("ThreadId").As<uint32>();
 		TSharedRef<FThreadState> ThreadState = GetThreadState(ThreadId);
 		uint64 LastCycle = ThreadState->LastCycle;
