Added documentation for Trace's IAnalyzer interface.

diff --git a/Engine/Source/Developer/TraceAnalysis/Public/Trace/Analyzer.h b/Engine/Source/Developer/TraceAnalysis/Public/Trace/Analyzer.h
--- a/Engine/Source/Developer/TraceAnalysis/Public/Trace/Analyzer.h
+++ b/Engine/Source/Developer/TraceAnalysis/Public/Trace/Analyzer.h
@@ -11,7 +11,15 @@
 namespace Trace
 {
 
-////////////////////////////////////////////////////////////////////////////////
+/**
+ * Interface that users implement to analyze the events in a trace. Analysis
+ * works by subscribing to events by name along with a user-provider "route"
+ * identifier. The IAnalyzer then receives callbacks when those events are
+ * encountered along with an interface to query the value of the event's fields.
+ *
+ * To analyze a trace, concrete IAnalyzer-derived objects are registered with a
+ * FAnalysisContext which is then asked to launch and coordinate the analysis.
+ */
 class TRACEANALYSIS_API IAnalyzer
 {
 public:
@@ -33,7 +41,11 @@ public:
 
 	struct FInterfaceBuilder
 	{
-		virtual void			RouteEvent(uint16 RouteId, const ANSICHAR* Logger, const ANSICHAR* Event) = 0;
+		/** Subscribe to an event required for analysis.
+		 * @param RouteId User-provided identifier for this event subscription.
+		 * @param Logger Name of the logger that emits the event.
+		 * @param Event Name of the event to subscribe to. */
+		virtual void RouteEvent(uint16 RouteId, const ANSICHAR* Logger, const ANSICHAR* Event) = 0;
 	};
 
 	struct FOnAnalysisContext
@@ -44,12 +56,19 @@ public:
 
 	struct FEventData
 	{
+		/** Queries the value of a field of the event. */
 		template <typename ValueType> ValueType GetValue(const ANSICHAR* FieldName) const;
-		virtual const uint8*	GetAttachment() const = 0;
-		virtual uint16			GetAttachmentSize() const = 0;
+
+		/** Returns the event's attachment. Not that this will always return an
+		 * address but if the event has no attachment then reading from that
+		 * address if undefined. */
+		virtual const uint8* GetAttachment() const = 0;
+
+		/** Returns the size of the events attachment, or 0 if none. */
+		virtual uint16 GetAttachmentSize() const = 0;
 
 	private:
-		virtual const void*		GetValueImpl(const ANSICHAR* FieldName, uint16& Type) const = 0;
+		virtual const void* GetValueImpl(const ANSICHAR* FieldName, uint16& Type) const = 0;
 	};
 
 	struct FOnEventContext
@@ -58,10 +77,22 @@ public:
 		const FEventData&		EventData;
 	};
 
-	virtual				~IAnalyzer() = default;
-	virtual void		OnAnalysisBegin(const FOnAnalysisContext& Context) = 0;
-	virtual void		OnAnalysisEnd() = 0;
-	virtual void		OnEvent(uint16 RouteId, const FOnEventContext& Context) = 0;
+	virtual ~IAnalyzer() = default;
+
+	/** Called when analysis of a trace is beginning. Analyzer implementers can
+	 * subscribe to the events that they are interested in at this point
+	 * @param Context Contextual information and interface for subscribing to events. */
+	virtual void OnAnalysisBegin(const FOnAnalysisContext& Context) = 0;
+
+	/** Indicates that the analysis of a trace log has completed and there are no
+	 * further events */
+	virtual void OnAnalysisEnd() = 0;
+
+	/** For each event subscribed to in OnAnalysisBegin(), the analysis engine
+	 * will call this method when those events are encountered in a trace log
+	 * @param RouteId User-provided identifier given when subscribing to a particular event
+	 * @param Context Access to the instance of the subscribed event */
+	virtual void OnEvent(uint16 RouteId, const FOnEventContext& Context) = 0;
 };
 
 ////////////////////////////////////////////////////////////////////////////////
