UnrealInsights:
Added file size and timestamp info to trace sessions

diff --git a/Engine/Source/Developer/TraceAnalysis/Private/Store/FileStore.cpp b/Engine/Source/Developer/TraceAnalysis/Private/Store/FileStore.cpp
--- a/Engine/Source/Developer/TraceAnalysis/Private/Store/FileStore.cpp
+++ b/Engine/Source/Developer/TraceAnalysis/Private/Store/FileStore.cpp
@@ -64,6 +64,8 @@ private:
 		FStoreSessionHandle Handle;
 		FString Name;
 		FString Path;
+		FDateTime TimeStamp;
+		uint64 Size;
 		bool bIsLive;
 		bool bIsValid;
 	};
@@ -184,17 +186,27 @@ FFileStore::~FFileStore()
 
 FFileStore::FSessionInfoInternal* FFileStore::AddSession(const FString& Path)
 {
+	IPlatformFile& FileSystem = IPlatformFile::GetPlatformPhysical();
+	FFileStatData FileStatData = FileSystem.GetStatData(*Path);
 	FScopeLock Lock(&SessionsCS);
-
 	FSessionInfoInternal* Session = new FSessionInfoInternal();
 	Session->Handle = NextSessionHandle++;
 	Session->Name = FPaths::GetBaseFilename(Path);
 	Session->Path = Path;
-	Session->bIsValid = true;
+	if (FileStatData.bIsValid)
+	{
+		Session->bIsValid = true;
+		Session->Size = uint64(FileStatData.FileSize);
+		Session->TimeStamp = FileStatData.CreationTime;
+	}
+	else
+	{
+		Session->bIsValid = false;
+		Session->Size = 0;
+	}
 	Session->bIsLive = false;
 	Sessions.Add(Session);
 	SessionsByPathMap.Add(Path, Session);
-
 	return Session;
 }
 
@@ -234,6 +246,8 @@ void FFileStore::GetAvailableSessions(TArray<FStoreSessionInfo>& OutSessions) co
 			SessionInfo.Handle = Session->Handle;
 			SessionInfo.Uri = *Session->Path;
 			SessionInfo.Name = *Session->Name;
+			SessionInfo.TimeStamp = Session->TimeStamp;
+			SessionInfo.Size = Session->Size;
 			SessionInfo.bIsLive = Session->bIsLive;
 		}
 	}
@@ -286,6 +300,20 @@ IInDataStream* FFileStore::OpenSessionStream(FStoreSessionHandle Handle)
 bool FFileStore::Tick(float DeltaTime)
 {
 	DirectoryWatcher->Tick(DeltaTime);
+
+	IPlatformFile& FileSystem = IPlatformFile::GetPlatformPhysical();
+	FScopeLock Lock(&SessionsCS);
+	for (FSessionInfoInternal* Session : Sessions)
+	{
+		if (Session->bIsValid && Session->bIsLive)
+		{
+			int64 NewSize = FileSystem.FileSize(*Session->Path);
+			if (NewSize >= 0)
+			{
+				Session->Size = NewSize;
+			}
+		}
+	}
 	return true;
 }
 
diff --git a/Engine/Source/Developer/TraceAnalysis/Public/Trace/Store.h b/Engine/Source/Developer/TraceAnalysis/Public/Trace/Store.h
--- a/Engine/Source/Developer/TraceAnalysis/Public/Trace/Store.h
+++ b/Engine/Source/Developer/TraceAnalysis/Public/Trace/Store.h
@@ -5,6 +5,7 @@
 #include "CoreTypes.h"
 #include "Containers/UnrealString.h"
 #include "Templates/SharedPointer.h"
+#include "Misc/DateTime.h"
 
 namespace Trace
 {
@@ -19,13 +20,17 @@ typedef uint64 FStoreSessionHandle;
 struct FStoreSessionInfo
 {
 	/** A handle to identify the trace */
-	FStoreSessionHandle Handle;
+	FStoreSessionHandle Handle = FStoreSessionHandle(-1);
 	/** Fully qualified path to the trace store item */
-	const TCHAR* Uri;
+	const TCHAR* Uri = nullptr;
 	/** Friendly-ish name of the trace */
-	const TCHAR* Name;
+	const TCHAR* Name = nullptr;
+	/** Timestamp of the trace */
+	FDateTime TimeStamp;
+	/** Size of the trace */
+	uint64 Size = 0;
 	/** True if the trace is currently being written to */
-	bool bIsLive;
+	bool bIsLive = false;
 };
 
 /**
diff --git a/Engine/Source/Developer/TraceServices/Private/SessionService.cpp b/Engine/Source/Developer/TraceServices/Private/SessionService.cpp
--- a/Engine/Source/Developer/TraceServices/Private/SessionService.cpp
+++ b/Engine/Source/Developer/TraceServices/Private/SessionService.cpp
@@ -103,6 +103,8 @@ bool FSessionService::GetSessionInfo(FSessionHandle SessionHandle, FSessionInfo&
 	}
 	OutSessionInfo.Uri = FindIt->Uri;
 	OutSessionInfo.Name = FindIt->Name;
+	OutSessionInfo.TimeStamp = FindIt->TimeStamp;
+	OutSessionInfo.Size = FindIt->Size;
 	OutSessionInfo.bIsLive = FindIt->bIsLive;
 	return true;
 }
@@ -246,6 +248,8 @@ void FSessionService::UpdateSessions()
 		SessionsToRemove.Remove(StoreSession.Handle);
 		Session.Uri = StoreSession.Uri;
 		Session.Name = StoreSession.Name;
+		Session.TimeStamp = StoreSession.TimeStamp;
+		Session.Size = StoreSession.Size;
 		Session.bIsLive = StoreSession.bIsLive;
 		Session.RecorderSessionHandle = 0;
 	}
diff --git a/Engine/Source/Developer/TraceServices/Private/SessionServicePrivate.h b/Engine/Source/Developer/TraceServices/Private/SessionServicePrivate.h
--- a/Engine/Source/Developer/TraceServices/Private/SessionServicePrivate.h
+++ b/Engine/Source/Developer/TraceServices/Private/SessionServicePrivate.h
@@ -41,6 +41,8 @@ private:
 		const TCHAR* Uri;
 		const TCHAR* Name;
 		TMap<FName, TSet<FString>> EnabledModuleLoggersMap;
+		FDateTime TimeStamp;
+		uint64 Size;
 		bool bIsLive;
 	};
 
diff --git a/Engine/Source/Developer/TraceServices/Public/TraceServices/SessionService.h b/Engine/Source/Developer/TraceServices/Public/TraceServices/SessionService.h
--- a/Engine/Source/Developer/TraceServices/Public/TraceServices/SessionService.h
+++ b/Engine/Source/Developer/TraceServices/Public/TraceServices/SessionService.h
@@ -14,9 +14,11 @@ typedef uint64 FSessionHandle;
 
 struct FSessionInfo
 {
-	const TCHAR* Uri;
-	const TCHAR* Name;
-	bool bIsLive;
+	const TCHAR* Uri = nullptr;
+	const TCHAR* Name = nullptr;
+	FDateTime TimeStamp;
+	uint64 Size = 0;
+	bool bIsLive = false;
 };
 
 class ISessionService
