Removed a little redundancy and regroup a few statements.

diff --git a/Engine/Source/Developer/TraceAnalysis/Private/Recorder/Recorder.cpp b/Engine/Source/Developer/TraceAnalysis/Private/Recorder/Recorder.cpp
--- a/Engine/Source/Developer/TraceAnalysis/Private/Recorder/Recorder.cpp
+++ b/Engine/Source/Developer/TraceAnalysis/Private/Recorder/Recorder.cpp
@@ -122,17 +122,23 @@ FRecorder::FSession* FRecorder::AcceptSession(FSocket& Socket)
 	}
 
 	TTuple<FStoreSessionHandle, IOutDataStream*> StoreSession = Store->CreateNewSession();
-	if (!StoreSession.Get<1>())
+	IOutDataStream* OutStream = StoreSession.Get<1>();
+	if (!OutStream)
 	{
 		return nullptr;
 	}
+
+	TSharedRef<FInternetAddr> PeerAddress = ISocketSubsystem::Get()->CreateInternetAddr();
+	Socket.GetPeerAddress(*PeerAddress);
+	PeerAddress->SetPort(1985);
+
 	FSession* Session = new FSession();
 	Session->ControlClientAddress = ISocketSubsystem::Get(PLATFORM_SOCKETSUBSYSTEM)->CreateInternetAddr();
 	Socket.GetPeerAddress(*Session->ControlClientAddress);
 	Session->ControlClientAddress->SetPort(1985);
 	Session->Socket = &Socket;
 	Session->StoreSessionHandle = StoreSession.Get<0>();
-	Session->StoreSessionStream = StoreSession.Get<1>();
+	Session->StoreSessionStream = OutStream;
 	Session->Thread = FRunnableThread::Create(Session, TEXT("TraceRecSession"));
 	return Session;
 }
