A missed diff from the last submit.

diff --git a/Engine/Source/Developer/TraceAnalysis/Private/Recorder/Recorder.cpp b/Engine/Source/Developer/TraceAnalysis/Private/Recorder/Recorder.cpp
--- a/Engine/Source/Developer/TraceAnalysis/Private/Recorder/Recorder.cpp
+++ b/Engine/Source/Developer/TraceAnalysis/Private/Recorder/Recorder.cpp
@@ -133,9 +133,7 @@ FRecorder::FSession* FRecorder::AcceptSession(FSocket& Socket)
 	PeerAddress->SetPort(1985);
 
 	FSession* Session = new FSession();
-	Session->ControlClientAddress = ISocketSubsystem::Get(PLATFORM_SOCKETSUBSYSTEM)->CreateInternetAddr();
-	Socket.GetPeerAddress(*Session->ControlClientAddress);
-	Session->ControlClientAddress->SetPort(1985);
+	Session->ControlClientAddress = PeerAddress;
 	Session->Socket = &Socket;
 	Session->StoreSessionHandle = StoreSession.Get<0>();
 	Session->StoreSessionStream = OutStream;
