Include the 'TRCE' handshake in a trace's output data stream so it can be treated agnostically regardless of how the data was received/written.

diff --git a/Engine/Source/Developer/TraceAnalysis/Private/Analysis/Engine.cpp b/Engine/Source/Developer/TraceAnalysis/Private/Analysis/Engine.cpp
--- a/Engine/Source/Developer/TraceAnalysis/Private/Analysis/Engine.cpp
+++ b/Engine/Source/Developer/TraceAnalysis/Private/Analysis/Engine.cpp
@@ -512,17 +512,32 @@ bool FAnalysisEngine::EstablishTransport(FStreamReader::FData& Data)
 		return false;
 	}
 
-	Data.Advance(sizeof(*Header));
+	// Check for the magic uint32. Early traces did not include this as it was
+	// used to validate a inbound socket connection and then discarded.
+	if (Header->Format == 'E')
+	{
+		const uint32* Magic = (const uint32*)(Data.GetPointer(sizeof(*Magic)));
+
+		if (*Magic == 'TRCE')
+		{
+			Data.Advance(sizeof(*Magic));
+			return EstablishTransport(Data);
+		}
+
+		return false;
+	}
 
 	switch (Header->Format)
 	{
-	case 1:		Transport = new FTransportReader();						break;
+	case 1:		Transport = new FTransportReader(); break;
+	default:	return false;
+	//case 'E':	/* See the magic above */ break;
 #if 0
 	case 4:		Transport = new FLz4TransportReader(Header->Parameter); break;
 #endif // 0
-	default:	return false;
 	}
 
+	Data.Advance(sizeof(*Header));
 	return true;
 }
 
diff --git a/Engine/Source/Developer/TraceAnalysis/Private/Recorder/Recorder.cpp b/Engine/Source/Developer/TraceAnalysis/Private/Recorder/Recorder.cpp
--- a/Engine/Source/Developer/TraceAnalysis/Private/Recorder/Recorder.cpp
+++ b/Engine/Source/Developer/TraceAnalysis/Private/Recorder/Recorder.cpp
@@ -105,10 +105,10 @@ FRecorder::~FRecorder()
 ////////////////////////////////////////////////////////////////////////////////
 FRecorder::FSession* FRecorder::AcceptSession(FSocket& Socket)
 {
+	uint32 Magic;
 	bool bAcceptable = false;
 	if (Socket.Wait(ESocketWaitConditions::WaitForRead, FTimespan(ETimespan::TicksPerSecond / 3)))
 	{
-		uint32 Magic;
 		int32 RecvSize;
 		if (Socket.Recv((uint8*)&Magic, sizeof(Magic), RecvSize))
 		{
@@ -128,6 +128,8 @@ FRecorder::FSession* FRecorder::AcceptSession(FSocket& Socket)
 		return nullptr;
 	}
 
+	OutStream->Write(&Magic, sizeof(Magic));
+
 	TSharedRef<FInternetAddr> PeerAddress = ISocketSubsystem::Get()->CreateInternetAddr();
 	Socket.GetPeerAddress(*PeerAddress);
 	PeerAddress->SetPort(1985);
