Simple detection and polite handling of big-endian traces should they appear in the future.

diff --git a/Engine/Source/Developer/TraceAnalysis/Private/Analysis/Engine.cpp b/Engine/Source/Developer/TraceAnalysis/Private/Analysis/Engine.cpp
--- a/Engine/Source/Developer/TraceAnalysis/Private/Analysis/Engine.cpp
+++ b/Engine/Source/Developer/TraceAnalysis/Private/Analysis/Engine.cpp
@@ -514,9 +514,14 @@ bool FAnalysisEngine::EstablishTransport(FStreamReader::FData& Data)
 
 	// Check for the magic uint32. Early traces did not include this as it was
 	// used to validate a inbound socket connection and then discarded.
-	if (Header->Format == 'E')
+	if (Header->Format == 'E' || Header->Format == 'T')
 	{
 		const uint32* Magic = (const uint32*)(Data.GetPointer(sizeof(*Magic)));
+		if (*Magic == 'ECRT')
+		{
+			// Source is big-endian which we don't currently support
+			return false;
+		}
 
 		if (*Magic == 'TRCE')
 		{
@@ -532,6 +537,7 @@ bool FAnalysisEngine::EstablishTransport(FStreamReader::FData& Data)
 	case 1:		Transport = new FTransportReader(); break;
 	default:	return false;
 	//case 'E':	/* See the magic above */ break;
+	//case 'T':	/* See the magic above */ break;
 #if 0
 	case 4:		Transport = new FLz4TransportReader(Header->Parameter); break;
 #endif // 0
diff --git a/Engine/Source/Developer/TraceAnalysis/Private/Recorder/Recorder.cpp b/Engine/Source/Developer/TraceAnalysis/Private/Recorder/Recorder.cpp
--- a/Engine/Source/Developer/TraceAnalysis/Private/Recorder/Recorder.cpp
+++ b/Engine/Source/Developer/TraceAnalysis/Private/Recorder/Recorder.cpp
@@ -112,7 +112,8 @@ FRecorder::FSession* FRecorder::AcceptSession(FSocket& Socket)
 		int32 RecvSize;
 		if (Socket.Recv((uint8*)&Magic, sizeof(Magic), RecvSize))
 		{
-			bAcceptable = (RecvSize == sizeof(Magic)) & (Magic == 'TRCE');
+			bAcceptable = (RecvSize == sizeof(Magic));
+			bAcceptable |= ((Magic == 'TRCE') | (Magic == 'ECRT'));
 		}
 	}
 
