It is possible that only a partial header arrives over a socket. To be tolerant of such a scenario now we ensure we've received enough bytes to reasonably constitute some header data.

diff --git a/Engine/Source/Developer/TraceAnalysis/Private/Analysis/Engine.cpp b/Engine/Source/Developer/TraceAnalysis/Private/Analysis/Engine.cpp
--- a/Engine/Source/Developer/TraceAnalysis/Private/Analysis/Engine.cpp
+++ b/Engine/Source/Developer/TraceAnalysis/Private/Analysis/Engine.cpp
@@ -552,6 +552,12 @@ bool FAnalysisEngine::OnData(FStreamReader::FData& Data)
 {
 	if (Transport == nullptr)
 	{
+		// Ensure we've a reasonable amount of data to establish the transport with
+		if (Data.GetPointer(32) == nullptr)
+		{
+			return true;
+		}
+
 		if (!EstablishTransport(Data))
 		{
 			return false;
