Moved the definition of FAnalysisEngine::FDispatch into the .cpp file.

diff --git a/Engine/Source/Developer/TraceAnalysis/Private/Analysis/Engine.cpp b/Engine/Source/Developer/TraceAnalysis/Private/Analysis/Engine.cpp
--- a/Engine/Source/Developer/TraceAnalysis/Private/Analysis/Engine.cpp
+++ b/Engine/Source/Developer/TraceAnalysis/Private/Analysis/Engine.cpp
@@ -212,6 +212,28 @@ bool FLz4TransportReader::NextChunk()
 
 
 
+////////////////////////////////////////////////////////////////////////////////
+struct FAnalysisEngine::FDispatch
+{
+	struct FField
+	{
+		uint32		Hash;
+		uint16		Offset;
+		uint16		Size;
+		uint16		_Unused0;
+		uint8		TypeInfo;
+		uint8		_Unused1;
+	};
+
+	uint16			FirstRoute;
+	uint16			FieldCount;
+	uint16			EventSize;
+	uint16			_Unused0;
+	FField			Fields[];
+};
+
+
+
 ////////////////////////////////////////////////////////////////////////////////
 struct FAnalysisEngine::FEventDataImpl
 	: public FEventData
diff --git a/Engine/Source/Developer/TraceAnalysis/Private/Analysis/Engine.h b/Engine/Source/Developer/TraceAnalysis/Private/Analysis/Engine.h
--- a/Engine/Source/Developer/TraceAnalysis/Private/Analysis/Engine.h
+++ b/Engine/Source/Developer/TraceAnalysis/Private/Analysis/Engine.h
@@ -26,25 +26,7 @@ public:
 	bool				OnData(FStreamReader::FData& Data);
 
 private:
-	struct FDispatch
-	{
-		struct FField
-		{
-			uint32		Hash;
-			uint16		Offset;
-			uint16		Size;
-			uint16		_Unused0;
-			uint8		TypeInfo;
-			uint8		_Unused1;
-		};
-
-		uint16			FirstRoute;
-		uint16			FieldCount;
-		uint16			EventSize;
-		uint16			_Unused0;
-		FField			Fields[];
-	};
-
+	struct				FDispatch;
 	struct FRoute
 	{
 		uint16			HashIndex;
