Store subscription hashes directly in routes instead of deduplicating into an additional array. Reduces complexity.

diff --git a/Engine/Source/Developer/TraceAnalysis/Private/Analysis/Engine.cpp b/Engine/Source/Developer/TraceAnalysis/Private/Analysis/Engine.cpp
--- a/Engine/Source/Developer/TraceAnalysis/Private/Analysis/Engine.cpp
+++ b/Engine/Source/Developer/TraceAnalysis/Private/Analysis/Engine.cpp
@@ -365,23 +365,9 @@ void FAnalysisEngine::AddRoute(
 {
 	check(AnalyzerIndex < Analyzers.Num());
 
-	uint16 HashIndex = 0;
-	for (uint16 n = Hashes.Num(); HashIndex < n; ++HashIndex)
-	{
-		if (Hashes[HashIndex] == Hash)
-		{
-			break;
-		}
-	}
-
-	if (HashIndex == Hashes.Num())
-	{
-		Hashes.Add(Hash);
-	}
-
 	FRoute& Route = Routes.Emplace_GetRef();
 	Route.Id = Id;
-	Route.HashIndex = HashIndex;
+	Route.Hash = Hash;
 	Route.Count = 1;
 	Route.AnalyzerIndex = AnalyzerIndex;
 }
@@ -408,14 +394,14 @@ void FAnalysisEngine::OnNewTrace(const FOnEventContext& Context)
 		Analyzers[i]->OnAnalysisBegin(OnAnalysisContext);
 	}
 
-	Algo::SortBy(Routes, [] (const FRoute& Route) { return Route.HashIndex; });
+	Algo::SortBy(Routes, [] (const FRoute& Route) { return Route.Hash; });
 
 	FRoute* Cursor = Routes.GetData();
 	Cursor->Count = 1;
 
 	for (uint16 i = 1, n = Routes.Num(); i < n; ++i)
 	{
-		if (Routes[i].HashIndex == Cursor->HashIndex)
+		if (Routes[i].Hash == Cursor->Hash)
 		{
 			Cursor->Count++;
 		}
@@ -428,7 +414,7 @@ void FAnalysisEngine::OnNewTrace(const FOnEventContext& Context)
 
 	// Add a terminal route for events that aren't subscribed to
 	FRoute& Route = Routes.Emplace_GetRef();
-	Route.HashIndex = Hashes.Num();
+	Route.Hash = ~0u;
 	Route.Count = 0;
 }
 
@@ -500,26 +486,16 @@ void FAnalysisEngine::OnNewEvent(const FOnEventContext& Context)
 	Algo::SortBy(Fields, [] (const auto& Field) { return Field.Hash; });
 
 	// Find routes that have subscribed to this event.
-	uint16 HashIndex = 0;
-	for (uint16 n = Hashes.Num(); HashIndex < n; ++HashIndex)
-	{
-		if (Hashes[HashIndex] == DispatchHash.Get())
-		{
-			break;
-		}
-	}
-
-	uint16 RouteIndex = 0;
-	for (uint16 n = Routes.Num(); RouteIndex < n; ++RouteIndex)
+	Dispatch.FirstRoute = Routes.Num() - 1;
+	for (uint16 i = 0, n = Routes.Num(); i < n; ++i)
 	{
-		if (Routes[RouteIndex].HashIndex == HashIndex)
+		if (Routes[i].Hash == DispatchHash.Get())
 		{
+			Dispatch.FirstRoute = i;
 			break;
 		}
 	}
 
-	check(RouteIndex < Routes.Num());
-	Dispatch.FirstRoute = RouteIndex;
 }
 
 ////////////////////////////////////////////////////////////////////////////////
diff --git a/Engine/Source/Developer/TraceAnalysis/Private/Analysis/Engine.h b/Engine/Source/Developer/TraceAnalysis/Private/Analysis/Engine.h
--- a/Engine/Source/Developer/TraceAnalysis/Private/Analysis/Engine.h
+++ b/Engine/Source/Developer/TraceAnalysis/Private/Analysis/Engine.h
@@ -29,10 +29,11 @@ private:
 	struct				FDispatch;
 	struct FRoute
 	{
-		uint16			HashIndex;
+		uint32			Hash;
 		int16			Count;
 		uint16			Id;
 		uint16			AnalyzerIndex;
+		uint16			_Unused0;
 	};
 
 	struct				FEventDataImpl;
@@ -49,7 +50,6 @@ private:
 	void				OnTiming(const FOnEventContext& Context);
 	void				OnNewEvent(const FOnEventContext& Context);
 	FSessionContext		SessionContext;
-	TArray<uint32>		Hashes;
 	TArray<FRoute>		Routes;
 	TArray<IAnalyzer*>	Analyzers;
 	TArray<FDispatch*>	Dispatches;
