Change some constants to be enums to give some context to other humans.

diff --git a/Engine/Source/Developer/TraceAnalysis/Private/Analysis/Engine.cpp b/Engine/Source/Developer/TraceAnalysis/Private/Analysis/Engine.cpp
--- a/Engine/Source/Developer/TraceAnalysis/Private/Analysis/Engine.cpp
+++ b/Engine/Source/Developer/TraceAnalysis/Private/Analysis/Engine.cpp
@@ -289,6 +289,13 @@ enum ERouteId : uint16
 	RouteId_Timing,
 };
 
+////////////////////////////////////////////////////////////////////////////////
+// This is used to influence the order of routes (routes are sorted by hash).
+enum EKnownRouteHashes : uint32
+{
+	RouteHash_NewEvent = 0, // must be 0 to match traces.
+};
+
 ////////////////////////////////////////////////////////////////////////////////
 FAnalysisEngine::FAnalysisEngine(TArray<IAnalyzer*>&& InAnalyzers)
 : Analyzers(MoveTemp(InAnalyzers))
@@ -302,7 +309,7 @@ FAnalysisEngine::FAnalysisEngine(TArray<IAnalyzer*>&& InAnalyzers)
 	// yet but are expecting.
 	FDispatch& NewEventDispatch = AddDispatch(uint16(FNewEventEvent::Uid), 0);
 	NewEventDispatch.FirstRoute = 0;
-	AddRoute(SelfIndex, RouteId_NewEvent, 0);
+	AddRoute(SelfIndex, RouteId_NewEvent, RouteHash_NewEvent);
 	AddRoute(SelfIndex, RouteId_NewTrace, "$Trace", "NewTrace");
 	AddRoute(SelfIndex, RouteId_Timing, "$Trace", "Timing");
 }
