Made analysis a little more efficient when reading events' fields by eliminating FEventData's vtable

diff --git a/Engine/Source/Developer/TraceAnalysis/Private/Analysis/Engine.cpp b/Engine/Source/Developer/TraceAnalysis/Private/Analysis/Engine.cpp
--- a/Engine/Source/Developer/TraceAnalysis/Private/Analysis/Engine.cpp
+++ b/Engine/Source/Developer/TraceAnalysis/Private/Analysis/Engine.cpp
@@ -235,32 +235,29 @@ struct FAnalysisEngine::FDispatch
 
 
 ////////////////////////////////////////////////////////////////////////////////
-struct FAnalysisEngine::FEventDataImpl
-	: public FEventData
-{
-	virtual					~FEventDataImpl() = default;
-	virtual const void*		GetValueImpl(const ANSICHAR* FieldName, uint16& Type) const override;
-	virtual const uint8*	GetAttachment() const override;
-	virtual uint16			GetAttachmentSize() const override;
-	const FDispatch*		Dispatch;
-	const uint8*			Ptr;
-	uint16					Size;
+struct FAnalysisEngine::FEventDataInfo
+{
+	const FDispatch&	Dispatch;
+	const uint8*		Ptr;
+	uint16				Size;
 };
 
 ////////////////////////////////////////////////////////////////////////////////
-const void* FAnalysisEngine::FEventDataImpl::GetValueImpl(const ANSICHAR* FieldName, uint16& Type) const
+const void* IAnalyzer::FEventData::GetValueImpl(const ANSICHAR* FieldName, uint16& Type) const
 {
+	const auto& Info = *(const FAnalysisEngine::FEventDataInfo*)this;
+
 	FFnv1aHash Hash;
 	Hash.Add(FieldName);
 	uint32 NameHash = Hash.Get();
 
-	for (int i = 0, n = Dispatch->FieldCount; i < n; ++i)
+	for (int i = 0, n = Info.Dispatch.FieldCount; i < n; ++i)
 	{
-		const auto& Field = Dispatch->Fields[i];
+		const auto& Field = Info.Dispatch.Fields[i];
 		if (Field.Hash == NameHash)
 		{
 			Type = Field.TypeInfo;
-			return (Ptr + Field.Offset);
+			return (Info.Ptr + Field.Offset);
 		}
 	}
 
@@ -268,15 +265,17 @@ const void* FAnalysisEngine::FEventDataImpl::GetValueImpl(const ANSICHAR* FieldN
 }
 
 ////////////////////////////////////////////////////////////////////////////////
-const uint8* FAnalysisEngine::FEventDataImpl::GetAttachment() const
+const uint8* IAnalyzer::FEventData::GetAttachment() const
 {
-	return Ptr + Dispatch->EventSize;
+	const auto& Info = *(const FAnalysisEngine::FEventDataInfo*)this;
+	return Info.Ptr + Info.Dispatch.EventSize;
 }
 
 ////////////////////////////////////////////////////////////////////////////////
-uint16 FAnalysisEngine::FEventDataImpl::GetAttachmentSize() const
+uint32 IAnalyzer::FEventData::GetAttachmentSize() const
 {
-	return Size - Dispatch->EventSize;
+	const auto& Info = *(const FAnalysisEngine::FEventDataInfo*)this;
+	return Info.Size - Info.Dispatch.EventSize;
 }
 
 
@@ -300,8 +299,6 @@ enum EKnownRouteHashes : uint32
 FAnalysisEngine::FAnalysisEngine(TArray<IAnalyzer*>&& InAnalyzers)
 : Analyzers(MoveTemp(InAnalyzers))
 {
-	EventDataImpl = new FEventDataImpl();
-
 	uint16 SelfIndex = Analyzers.Num();
 	Analyzers.Add(this);
 
@@ -326,8 +323,6 @@ FAnalysisEngine::~FAnalysisEngine()
 	{
 		FMemory::Free(Dispatch);
 	}
-
-	delete EventDataImpl;
 }
 
 ////////////////////////////////////////////////////////////////////////////////
@@ -456,7 +451,7 @@ FAnalysisEngine::FDispatch& FAnalysisEngine::AddDispatch(uint16 Uid, uint16 Fiel
 ////////////////////////////////////////////////////////////////////////////////
 void FAnalysisEngine::OnNewEvent(const FOnEventContext& Context)
 {
-	const FEventDataImpl& EventData = (const FEventDataImpl&)(Context.EventData);
+	const FEventDataInfo& EventData = (const FEventDataInfo&)(Context.EventData);
 	const auto& NewEvent = *(FNewEventEvent*)(EventData.Ptr);
 
 	FDispatch& Dispatch = AddDispatch(NewEvent.EventUid, NewEvent.FieldCount);
@@ -603,15 +598,14 @@ bool FAnalysisEngine::OnData(FStreamReader::FData& Data)
 
 		const FDispatch* Dispatch = Dispatches[Uid];
 
-		EventDataImpl->Dispatch = Dispatch;
-		EventDataImpl->Ptr = Header->EventData;
-		EventDataImpl->Size = Header->Size;
+		FEventDataInfo EventDataInfo = { *Dispatch, Header->EventData, Header->Size };
+		const FEventData& EventData = (FEventData&)EventDataInfo;
 
 		const FRoute* Route = Routes.GetData() + Dispatch->FirstRoute;
 		for (uint32 n = Route->Count; n--; ++Route)
 		{
 			IAnalyzer* Analyzer = Analyzers[Route->AnalyzerIndex];
-			Analyzer->OnEvent(Route->Id, { SessionContext, *EventDataImpl });
+			Analyzer->OnEvent(Route->Id, { SessionContext, EventData });
 		}
 	}
 
diff --git a/Engine/Source/Developer/TraceAnalysis/Private/Analysis/Engine.h b/Engine/Source/Developer/TraceAnalysis/Private/Analysis/Engine.h
--- a/Engine/Source/Developer/TraceAnalysis/Private/Analysis/Engine.h
+++ b/Engine/Source/Developer/TraceAnalysis/Private/Analysis/Engine.h
@@ -21,12 +21,13 @@ class FAnalysisEngine
 	: public IAnalyzer
 {
 public:
+	struct				FDispatch;
+	struct				FEventDataInfo;
 						FAnalysisEngine(TArray<IAnalyzer*>&& InAnalyzers);
 						~FAnalysisEngine();
 	bool				OnData(FStreamReader::FData& Data);
 
 private:
-	struct				FDispatch;
 	struct FRoute
 	{
 		uint32			Hash;
@@ -36,8 +37,6 @@ private:
 		uint16			_Unused0;
 	};
 
-	struct				FEventDataImpl;
-
 	virtual void		OnAnalysisBegin(const FOnAnalysisContext& Context) override;
 	virtual void		OnAnalysisEnd() override;
 	virtual void		OnEvent(uint16 RouteId, const FOnEventContext& Context) override;
@@ -54,7 +53,6 @@ private:
 	TArray<IAnalyzer*>	Analyzers;
 	TArray<FDispatch*>	Dispatches;
 	FTransportReader*	Transport = nullptr;
-	FEventDataImpl*		EventDataImpl;
 };
 
 } // namespace Trace
diff --git a/Engine/Source/Developer/TraceAnalysis/Public/Trace/Analyzer.h b/Engine/Source/Developer/TraceAnalysis/Public/Trace/Analyzer.h
--- a/Engine/Source/Developer/TraceAnalysis/Public/Trace/Analyzer.h
+++ b/Engine/Source/Developer/TraceAnalysis/Public/Trace/Analyzer.h
@@ -54,7 +54,7 @@ public:
 		FInterfaceBuilder&		InterfaceBuilder;
 	};
 
-	struct FEventData
+	struct TRACEANALYSIS_API FEventData
 	{
 		/** Queries the value of a field of the event. */
 		template <typename ValueType> ValueType GetValue(const ANSICHAR* FieldName) const;
@@ -62,13 +62,13 @@ public:
 		/** Returns the event's attachment. Not that this will always return an
 		 * address but if the event has no attachment then reading from that
 		 * address if undefined. */
-		virtual const uint8* GetAttachment() const = 0;
+		const uint8* GetAttachment() const;
 
 		/** Returns the size of the events attachment, or 0 if none. */
-		virtual uint16 GetAttachmentSize() const = 0;
+		uint32 GetAttachmentSize() const;
 
 	private:
-		virtual const void* GetValueImpl(const ANSICHAR* FieldName, uint16& Type) const = 0;
+		const void* GetValueImpl(const ANSICHAR* FieldName, uint16& Type) const;
 	};
 
 	struct FOnEventContext
