Changed the IAnalysis interface to be optional as not everyone is interested in all calls made in to analyzers.

diff --git a/Engine/Source/Developer/TraceAnalysis/Private/Analysis/Engine.cpp b/Engine/Source/Developer/TraceAnalysis/Private/Analysis/Engine.cpp
--- a/Engine/Source/Developer/TraceAnalysis/Private/Analysis/Engine.cpp
+++ b/Engine/Source/Developer/TraceAnalysis/Private/Analysis/Engine.cpp
@@ -325,16 +325,6 @@ FAnalysisEngine::~FAnalysisEngine()
 	}
 }
 
-////////////////////////////////////////////////////////////////////////////////
-void FAnalysisEngine::OnAnalysisBegin(const FOnAnalysisContext& Context)
-{
-}
-
-////////////////////////////////////////////////////////////////////////////////
-void FAnalysisEngine::OnAnalysisEnd()
-{
-}
-
 ////////////////////////////////////////////////////////////////////////////////
 void FAnalysisEngine::OnEvent(uint16 RouteId, const FOnEventContext& Context)
 {
diff --git a/Engine/Source/Developer/TraceAnalysis/Private/Analysis/Engine.h b/Engine/Source/Developer/TraceAnalysis/Private/Analysis/Engine.h
--- a/Engine/Source/Developer/TraceAnalysis/Private/Analysis/Engine.h
+++ b/Engine/Source/Developer/TraceAnalysis/Private/Analysis/Engine.h
@@ -37,8 +37,6 @@ private:
 		uint16			_Unused0;
 	};
 
-	virtual void		OnAnalysisBegin(const FOnAnalysisContext& Context) override;
-	virtual void		OnAnalysisEnd() override;
 	virtual void		OnEvent(uint16 RouteId, const FOnEventContext& Context) override;
 
 	bool				EstablishTransport(FStreamReader::FData& Data);
diff --git a/Engine/Source/Developer/TraceAnalysis/Public/Trace/Analyzer.h b/Engine/Source/Developer/TraceAnalysis/Public/Trace/Analyzer.h
--- a/Engine/Source/Developer/TraceAnalysis/Public/Trace/Analyzer.h
+++ b/Engine/Source/Developer/TraceAnalysis/Public/Trace/Analyzer.h
@@ -82,17 +82,17 @@ public:
 	/** Called when analysis of a trace is beginning. Analyzer implementers can
 	 * subscribe to the events that they are interested in at this point
 	 * @param Context Contextual information and interface for subscribing to events. */
-	virtual void OnAnalysisBegin(const FOnAnalysisContext& Context) = 0;
+	virtual void OnAnalysisBegin(const FOnAnalysisContext& Context) {}
 
 	/** Indicates that the analysis of a trace log has completed and there are no
 	 * further events */
-	virtual void OnAnalysisEnd() = 0;
+	virtual void OnAnalysisEnd() {}
 
 	/** For each event subscribed to in OnAnalysisBegin(), the analysis engine
 	 * will call this method when those events are encountered in a trace log
 	 * @param RouteId User-provided identifier given when subscribing to a particular event
 	 * @param Context Access to the instance of the subscribed event */
-	virtual void OnEvent(uint16 RouteId, const FOnEventContext& Context) = 0;
+	virtual void OnEvent(uint16 RouteId, const FOnEventContext& Context) {}
 };
 
 ////////////////////////////////////////////////////////////////////////////////
diff --git a/Engine/Source/Developer/TraceServices/Private/Analyzers/CountersTraceAnalysis.cpp b/Engine/Source/Developer/TraceServices/Private/Analyzers/CountersTraceAnalysis.cpp
--- a/Engine/Source/Developer/TraceServices/Private/Analyzers/CountersTraceAnalysis.cpp
+++ b/Engine/Source/Developer/TraceServices/Private/Analyzers/CountersTraceAnalysis.cpp
@@ -19,10 +19,6 @@ void FCountersAnalyzer::OnAnalysisBegin(const FOnAnalysisContext& Context)
 	Builder.RouteEvent(RouteId_SetValueFloat, "Counters", "SetValueFloat");
 }
 
-void FCountersAnalyzer::OnAnalysisEnd()
-{
-}
-
 void FCountersAnalyzer::OnEvent(uint16 RouteId, const FOnEventContext& Context)
 {
 	Trace::FAnalysisSessionEditScope _(Session);
diff --git a/Engine/Source/Developer/TraceServices/Private/Analyzers/CountersTraceAnalysis.h b/Engine/Source/Developer/TraceServices/Private/Analyzers/CountersTraceAnalysis.h
--- a/Engine/Source/Developer/TraceServices/Private/Analyzers/CountersTraceAnalysis.h
+++ b/Engine/Source/Developer/TraceServices/Private/Analyzers/CountersTraceAnalysis.h
@@ -20,7 +20,6 @@ public:
 	FCountersAnalyzer(Trace::IAnalysisSession& Session, Trace::ICounterProvider& CounterProvider);
 	virtual void OnAnalysisBegin(const FOnAnalysisContext& Context) override;
 	virtual void OnEvent(uint16 RouteId, const FOnEventContext& Context) override;
-	virtual void OnAnalysisEnd() override;
 
 private:
 	enum : uint16
diff --git a/Engine/Source/Developer/TraceServices/Private/Analyzers/CpuProfilerTraceAnalysis.cpp b/Engine/Source/Developer/TraceServices/Private/Analyzers/CpuProfilerTraceAnalysis.cpp
--- a/Engine/Source/Developer/TraceServices/Private/Analyzers/CpuProfilerTraceAnalysis.cpp
+++ b/Engine/Source/Developer/TraceServices/Private/Analyzers/CpuProfilerTraceAnalysis.cpp
@@ -20,10 +20,6 @@ void FCpuProfilerAnalyzer::OnAnalysisBegin(const FOnAnalysisContext& Context)
 	Builder.RouteEvent(RouteId_EndCapture, "CpuProfiler", "EndCapture");
 }
 
-void FCpuProfilerAnalyzer::OnAnalysisEnd()
-{
-}
-
 void FCpuProfilerAnalyzer::OnEvent(uint16 RouteId, const FOnEventContext& Context)
 {
 	Trace::FAnalysisSessionEditScope _(Session);
diff --git a/Engine/Source/Developer/TraceServices/Private/Analyzers/CpuProfilerTraceAnalysis.h b/Engine/Source/Developer/TraceServices/Private/Analyzers/CpuProfilerTraceAnalysis.h
--- a/Engine/Source/Developer/TraceServices/Private/Analyzers/CpuProfilerTraceAnalysis.h
+++ b/Engine/Source/Developer/TraceServices/Private/Analyzers/CpuProfilerTraceAnalysis.h
@@ -19,7 +19,6 @@ public:
 	FCpuProfilerAnalyzer(Trace::IAnalysisSession& Session, Trace::FTimingProfilerProvider& TimingProfilerProvider);
 	virtual void OnAnalysisBegin(const FOnAnalysisContext& Context) override;
 	virtual void OnEvent(uint16 RouteId, const FOnEventContext& Context) override;
-	virtual void OnAnalysisEnd() override;
 
 private:
 	struct EventScopeState
diff --git a/Engine/Source/Developer/TraceServices/Private/Analyzers/GpuProfilerTraceAnalysis.h b/Engine/Source/Developer/TraceServices/Private/Analyzers/GpuProfilerTraceAnalysis.h
--- a/Engine/Source/Developer/TraceServices/Private/Analyzers/GpuProfilerTraceAnalysis.h
+++ b/Engine/Source/Developer/TraceServices/Private/Analyzers/GpuProfilerTraceAnalysis.h
@@ -18,7 +18,6 @@ public:
 	FGpuProfilerAnalyzer(Trace::FAnalysisSession& Session, Trace::FTimingProfilerProvider& TimingProfilerProvider);
 	virtual void OnAnalysisBegin(const FOnAnalysisContext& Context) override;
 	virtual void OnEvent(uint16 RouteId, const FOnEventContext& Context) override;
-	virtual void OnAnalysisEnd() override {};
 
 private:
 	enum : uint16
diff --git a/Engine/Source/Developer/TraceServices/Private/Analyzers/LoadTimeTraceAnalysis.cpp b/Engine/Source/Developer/TraceServices/Private/Analyzers/LoadTimeTraceAnalysis.cpp
--- a/Engine/Source/Developer/TraceServices/Private/Analyzers/LoadTimeTraceAnalysis.cpp
+++ b/Engine/Source/Developer/TraceServices/Private/Analyzers/LoadTimeTraceAnalysis.cpp
@@ -89,11 +89,6 @@ void FAsyncLoadingTraceAnalyzer::OnAnalysisBegin(const FOnAnalysisContext& Conte
 	Builder.RouteEvent(RouteId_ClassInfo, "LoadTime", "ClassInfo");
 }
 
-void FAsyncLoadingTraceAnalyzer::OnAnalysisEnd()
-{
-	
-}
-
 void FAsyncLoadingTraceAnalyzer::OnEvent(uint16 RouteId, const FOnEventContext& Context)
 {
 	const auto& EventData = Context.EventData;
@@ -429,4 +424,4 @@ void FAsyncLoadingTraceAnalyzer::PackageRequestAssociation(const FOnEventContext
 		}
 		LoadRequest->Packages.Add(AsyncPackageState->PackageInfo);
 	}
-}
\ No newline at end of file
+}
diff --git a/Engine/Source/Developer/TraceServices/Private/Analyzers/LoadTimeTraceAnalysis.h b/Engine/Source/Developer/TraceServices/Private/Analyzers/LoadTimeTraceAnalysis.h
--- a/Engine/Source/Developer/TraceServices/Private/Analyzers/LoadTimeTraceAnalysis.h
+++ b/Engine/Source/Developer/TraceServices/Private/Analyzers/LoadTimeTraceAnalysis.h
@@ -28,7 +28,6 @@ public:
 
 	virtual void OnAnalysisBegin(const FOnAnalysisContext& Context) override;
 	virtual void OnEvent(uint16 RouteId, const FOnEventContext& Context) override;
-	virtual void OnAnalysisEnd() override;
 	
 private:
 	struct FRequestState;
diff --git a/Engine/Source/Developer/TraceServices/Private/Analyzers/LogTraceAnalysis.h b/Engine/Source/Developer/TraceServices/Private/Analyzers/LogTraceAnalysis.h
--- a/Engine/Source/Developer/TraceServices/Private/Analyzers/LogTraceAnalysis.h
+++ b/Engine/Source/Developer/TraceServices/Private/Analyzers/LogTraceAnalysis.h
@@ -18,7 +18,6 @@ public:
 	FLogTraceAnalyzer(Trace::IAnalysisSession& Session, Trace::FLogProvider& LogProvider);
 	virtual void OnAnalysisBegin(const FOnAnalysisContext& Context) override;
 	virtual void OnEvent(uint16 RouteId, const FOnEventContext& Context) override;
-	virtual void OnAnalysisEnd() override {};
 
 private:
 	enum : uint16
diff --git a/Engine/Source/Developer/TraceServices/Private/Analyzers/MiscTraceAnalysis.h b/Engine/Source/Developer/TraceServices/Private/Analyzers/MiscTraceAnalysis.h
--- a/Engine/Source/Developer/TraceServices/Private/Analyzers/MiscTraceAnalysis.h
+++ b/Engine/Source/Developer/TraceServices/Private/Analyzers/MiscTraceAnalysis.h
@@ -28,7 +28,6 @@ public:
 					   Trace::FFrameProvider& FrameProvider);
 	virtual void OnAnalysisBegin(const FOnAnalysisContext& Context) override;
 	virtual void OnEvent(uint16 RouteId, const FOnEventContext& Context) override;
-	virtual void OnAnalysisEnd() override {};
 
 private:
 	enum : uint16
diff --git a/Engine/Source/Developer/TraceServices/Private/Analyzers/PlatformFileTraceAnalysis.cpp b/Engine/Source/Developer/TraceServices/Private/Analyzers/PlatformFileTraceAnalysis.cpp
--- a/Engine/Source/Developer/TraceServices/Private/Analyzers/PlatformFileTraceAnalysis.cpp
+++ b/Engine/Source/Developer/TraceServices/Private/Analyzers/PlatformFileTraceAnalysis.cpp
@@ -171,7 +171,3 @@ void FPlatformFileTraceAnalyzer::OnEvent(uint16 RouteId, const FOnEventContext&
 	}
 	}
 }
-
-void FPlatformFileTraceAnalyzer::OnAnalysisEnd()
-{
-}
\ No newline at end of file
diff --git a/Engine/Source/Developer/TraceServices/Private/Analyzers/PlatformFileTraceAnalysis.h b/Engine/Source/Developer/TraceServices/Private/Analyzers/PlatformFileTraceAnalysis.h
--- a/Engine/Source/Developer/TraceServices/Private/Analyzers/PlatformFileTraceAnalysis.h
+++ b/Engine/Source/Developer/TraceServices/Private/Analyzers/PlatformFileTraceAnalysis.h
@@ -20,7 +20,6 @@ public:
 	FPlatformFileTraceAnalyzer(Trace::IAnalysisSession& Session, Trace::FFileActivityProvider& FileActivityProvider);
 	virtual void OnAnalysisBegin(const FOnAnalysisContext& Context) override;
 	virtual void OnEvent(uint16 RouteId, const FOnEventContext& Context) override;
-	virtual void OnAnalysisEnd() override;
 
 private:
 	enum : uint16
diff --git a/Engine/Source/Developer/TraceServices/Private/Analyzers/StatsTraceAnalysis.cpp b/Engine/Source/Developer/TraceServices/Private/Analyzers/StatsTraceAnalysis.cpp
--- a/Engine/Source/Developer/TraceServices/Private/Analyzers/StatsTraceAnalysis.cpp
+++ b/Engine/Source/Developer/TraceServices/Private/Analyzers/StatsTraceAnalysis.cpp
@@ -19,10 +19,6 @@ void FStatsAnalyzer::OnAnalysisBegin(const FOnAnalysisContext& Context)
 	Builder.RouteEvent(RouteId_EventBatch, "Stats", "EventBatch");
 }
 
-void FStatsAnalyzer::OnAnalysisEnd()
-{
-}
-
 void FStatsAnalyzer::OnEvent(uint16 RouteId, const FOnEventContext& Context)
 {
 	Trace::FAnalysisSessionEditScope _(Session);
diff --git a/Engine/Source/Developer/TraceServices/Private/Analyzers/StatsTraceAnalysis.h b/Engine/Source/Developer/TraceServices/Private/Analyzers/StatsTraceAnalysis.h
--- a/Engine/Source/Developer/TraceServices/Private/Analyzers/StatsTraceAnalysis.h
+++ b/Engine/Source/Developer/TraceServices/Private/Analyzers/StatsTraceAnalysis.h
@@ -20,7 +20,6 @@ public:
 	FStatsAnalyzer(Trace::IAnalysisSession& Session, Trace::ICounterProvider& CounterProvider);
 	virtual void OnAnalysisBegin(const FOnAnalysisContext& Context) override;
 	virtual void OnEvent(uint16 RouteId, const FOnEventContext& Context) override;
-	virtual void OnAnalysisEnd() override;
 
 private:
 	enum : uint16
