Store various names associated with an event when they appear in the event stream

diff --git a/Engine/Source/Developer/TraceAnalysis/Private/Analysis/Engine.cpp b/Engine/Source/Developer/TraceAnalysis/Private/Analysis/Engine.cpp
--- a/Engine/Source/Developer/TraceAnalysis/Private/Analysis/Engine.cpp
+++ b/Engine/Source/Developer/TraceAnalysis/Private/Analysis/Engine.cpp
@@ -220,14 +220,16 @@ struct FAnalysisEngine::FDispatch
 		uint32		Hash;
 		uint16		Offset;
 		uint16		Size;
-		uint16		_Unused0;
+		uint16		NameOffset;			// From FField ptr
 		uint8		TypeInfo;
-		uint8		_Unused1;
+		uint8		_Unused0;
 	};
 
+	uint16			Uid;
 	uint16			FirstRoute;
-	uint16			FieldCount;
+	uint16			FieldCount;			// Implicit logger name offset; Fields + FieldCount
 	uint16			EventSize;
+	uint16			EventNameOffset;	// From FDispatch ptr
 	uint16			_Unused0;
 	FField			Fields[];
 };
@@ -418,10 +420,13 @@ void FAnalysisEngine::OnTiming(const FOnEventContext& Context)
 }
 
 ////////////////////////////////////////////////////////////////////////////////
-FAnalysisEngine::FDispatch& FAnalysisEngine::AddDispatch(uint16 Uid, uint16 FieldCount)
+FAnalysisEngine::FDispatch& FAnalysisEngine::AddDispatch(
+	uint16 Uid,
+	uint16 FieldCount,
+	uint16 ExtraData)
 {
 	// Allocate a block of memory to hold the dispatch
-	uint32 Size = sizeof(FDispatch) + (sizeof(FDispatch::FField) * FieldCount);
+	uint32 Size = sizeof(FDispatch) + (sizeof(FDispatch::FField) * FieldCount) + ExtraData;
 	auto* Dispatch = (FDispatch*)FMemory::Malloc(Size);
 	Dispatch->FieldCount = FieldCount;
 	Dispatch->EventSize = 0;
@@ -444,7 +449,14 @@ void FAnalysisEngine::OnNewEvent(const FOnEventContext& Context)
 	const FEventDataInfo& EventData = (const FEventDataInfo&)(Context.EventData);
 	const auto& NewEvent = *(FNewEventEvent*)(EventData.Ptr);
 
-	FDispatch& Dispatch = AddDispatch(NewEvent.EventUid, NewEvent.FieldCount);
+	// Create a new dispatch with enough space to store new event's various names
+	uint16 NameDataSize = NewEvent.LoggerNameSize + NewEvent.EventNameSize;
+	NameDataSize += NewEvent.FieldCount + 2; // null terminators
+	for (int i = 0, n = NewEvent.FieldCount; i < n; ++i)
+	{
+		NameDataSize += NewEvent.Fields[i].NameSize + 1;
+	}
+	FDispatch& Dispatch = AddDispatch(NewEvent.EventUid, NewEvent.FieldCount, NameDataSize);
 
 	if (NewEvent.FieldCount)
 	{
@@ -452,7 +464,7 @@ void FAnalysisEngine::OnNewEvent(const FOnEventContext& Context)
 		Dispatch.EventSize = LastField.Offset + LastField.Size;
 	}
 
-	const uint8* NameCursor = (uint8*)(NewEvent.Fields + NewEvent.FieldCount);
+	const uint8* NameCursor = (const uint8*)(NewEvent.Fields + NewEvent.FieldCount);
 
 	// Calculate this dispatch's hash.
 	FFnv1aHash DispatchHash;
@@ -474,9 +486,43 @@ void FAnalysisEngine::OnNewEvent(const FOnEventContext& Context)
 		Out.TypeInfo = In.TypeInfo;
 	}
 
+	// Write out names with null terminators.
+	NameCursor = (const uint8*)(NewEvent.Fields + NewEvent.FieldCount);
+	uint8* WriteCursor = (uint8*)(Dispatch.Fields + Dispatch.FieldCount);
+	auto WriteName = [&] (uint32 Size)
+	{
+		memcpy(WriteCursor, NameCursor, Size);
+		NameCursor += Size;
+		WriteCursor += Size + 1;
+		WriteCursor[-1] = '\0';
+	};
+
+	UPTRINT EventNameOffset = UPTRINT(Dispatch.Fields + NewEvent.FieldCount);
+	EventNameOffset -= UPTRINT(&Dispatch);
+	EventNameOffset += NewEvent.LoggerNameSize + 1;
+	Dispatch.EventNameOffset = uint16(EventNameOffset);
+
+	WriteName(NewEvent.LoggerNameSize);
+	WriteName(NewEvent.EventNameSize);
+	for (int i = 0, n = NewEvent.FieldCount; i < n; ++i)
+	{
+		auto& Out = Dispatch.Fields[i];
+		Out.NameOffset = uint16(UPTRINT(WriteCursor) - UPTRINT(&Dispatch));
+
+		WriteName(NewEvent.Fields[i].NameSize);
+	}
+
+	// Sort by hash so we can binary search when looking up.
 	TArrayView<FDispatch::FField> Fields(Dispatch.Fields, Dispatch.FieldCount);
 	Algo::SortBy(Fields, [] (const auto& Field) { return Field.Hash; });
 
+	// Fix up field name offsets
+	for (int i = 0, n = NewEvent.FieldCount; i < n; ++i)
+	{
+		auto& Out = Dispatch.Fields[i];
+		Out.NameOffset = uint16(UPTRINT(&Dispatch) + Out.NameOffset - UPTRINT(&Out));
+	}
+
 	// Find routes that have subscribed to this event.
 	Dispatch.FirstRoute = Routes.Num() - 1;
 	for (uint16 i = 0, n = Routes.Num(); i < n; ++i)
diff --git a/Engine/Source/Developer/TraceAnalysis/Private/Analysis/Engine.h b/Engine/Source/Developer/TraceAnalysis/Private/Analysis/Engine.h
--- a/Engine/Source/Developer/TraceAnalysis/Private/Analysis/Engine.h
+++ b/Engine/Source/Developer/TraceAnalysis/Private/Analysis/Engine.h
@@ -40,7 +40,7 @@ private:
 	virtual void		OnEvent(uint16 RouteId, const FOnEventContext& Context) override;
 
 	bool				EstablishTransport(FStreamReader::FData& Data);
-	FDispatch&			AddDispatch(uint16 Uid, uint16 FieldCount);
+	FDispatch&			AddDispatch(uint16 Uid, uint16 FieldCount=0, uint16 ExtraData=0);
 	void				AddRoute(uint16 AnalyzerIndex, uint16 Id, const ANSICHAR* Logger, const ANSICHAR* Event);
 	void				AddRoute(uint16 AnalyzerIndex, uint16 Id, uint32 Hash);
 	void				OnNewTrace(const FOnEventContext& Context);
