Added a public API to be able to query type/reflection information about events in a trace stream.

diff --git a/Engine/Source/Developer/TraceAnalysis/Private/Analysis/Engine.cpp b/Engine/Source/Developer/TraceAnalysis/Private/Analysis/Engine.cpp
--- a/Engine/Source/Developer/TraceAnalysis/Private/Analysis/Engine.cpp
+++ b/Engine/Source/Developer/TraceAnalysis/Private/Analysis/Engine.cpp
@@ -236,6 +236,90 @@ struct FAnalysisEngine::FDispatch
 
 
 
+////////////////////////////////////////////////////////////////////////////////
+uint32 IAnalyzer::FEventTypeInfo::GetId() const
+{
+	const auto& Inner = *(const FAnalysisEngine::FDispatch*)this;
+	return Inner.Uid;
+}
+
+////////////////////////////////////////////////////////////////////////////////
+uint32 IAnalyzer::FEventTypeInfo::GetSize() const
+{
+	const auto& Inner = *(const FAnalysisEngine::FDispatch*)this;
+	return Inner.EventSize;
+}
+
+////////////////////////////////////////////////////////////////////////////////
+const ANSICHAR* IAnalyzer::FEventTypeInfo::GetName() const
+{
+	const auto& Inner = *(const FAnalysisEngine::FDispatch*)this;
+	return (const ANSICHAR*)(UPTRINT(&Inner) + Inner.EventNameOffset);
+}
+
+////////////////////////////////////////////////////////////////////////////////
+const ANSICHAR* IAnalyzer::FEventTypeInfo::GetLoggerName() const
+{
+	const auto& Inner = *(const FAnalysisEngine::FDispatch*)this;
+	return (const ANSICHAR*)(Inner.Fields + Inner.FieldCount);
+}
+
+////////////////////////////////////////////////////////////////////////////////
+uint32 IAnalyzer::FEventTypeInfo::GetFieldCount() const
+{
+	const auto& Inner = *(const FAnalysisEngine::FDispatch*)this;
+	return Inner.FieldCount;
+}
+
+////////////////////////////////////////////////////////////////////////////////
+const IAnalyzer::FEventFieldInfo* IAnalyzer::FEventTypeInfo::GetFieldInfo(uint32 Index) const
+{
+	if (Index >= GetFieldCount())
+	{
+		return nullptr;
+	}
+
+	const auto& Inner = *(const FAnalysisEngine::FDispatch*)this;
+	return (const IAnalyzer::FEventFieldInfo*)(Inner.Fields + Index);
+}
+
+
+
+////////////////////////////////////////////////////////////////////////////////
+const ANSICHAR* IAnalyzer::FEventFieldInfo::GetName() const
+{
+	const auto* Inner = (const FAnalysisEngine::FDispatch::FField*)this;
+	return (const ANSICHAR*)(UPTRINT(Inner) + Inner->NameOffset);
+}
+
+////////////////////////////////////////////////////////////////////////////////
+uint32 IAnalyzer::FEventFieldInfo::GetOffset() const
+{
+	const auto* Inner = (const FAnalysisEngine::FDispatch::FField*)this;
+	return Inner->Offset;
+}
+
+////////////////////////////////////////////////////////////////////////////////
+uint32 IAnalyzer::FEventFieldInfo::GetSize() const
+{
+	const auto* Inner = (const FAnalysisEngine::FDispatch::FField*)this;
+	return 1 << (Inner->TypeInfo & _Field_Pow2SizeMask);
+}
+
+////////////////////////////////////////////////////////////////////////////////
+IAnalyzer::FEventFieldInfo::EType IAnalyzer::FEventFieldInfo::GetType() const
+{
+	const auto* Inner = (const FAnalysisEngine::FDispatch::FField*)this;
+	switch (Inner->TypeInfo & _Field_CategoryMask)
+	{
+	case _Field_Integer:	return EType::Integer;
+	case _Field_Float:		return EType::Float;
+	}
+	return EType::None;
+}
+
+
+
 ////////////////////////////////////////////////////////////////////////////////
 struct FAnalysisEngine::FEventDataInfo
 {
@@ -244,6 +328,13 @@ struct FAnalysisEngine::FEventDataInfo
 	uint16				Size;
 };
 
+////////////////////////////////////////////////////////////////////////////////
+const IAnalyzer::FEventTypeInfo& IAnalyzer::FEventData::GetTypeInfo() const
+{
+	const auto& Info = *(const FAnalysisEngine::FEventDataInfo*)this;
+	return (const FEventTypeInfo&)(Info.Dispatch);
+}
+
 ////////////////////////////////////////////////////////////////////////////////
 const void* IAnalyzer::FEventData::GetValueImpl(const ANSICHAR* FieldName, uint16& Type) const
 {
@@ -428,6 +519,7 @@ FAnalysisEngine::FDispatch& FAnalysisEngine::AddDispatch(
 	// Allocate a block of memory to hold the dispatch
 	uint32 Size = sizeof(FDispatch) + (sizeof(FDispatch::FField) * FieldCount) + ExtraData;
 	auto* Dispatch = (FDispatch*)FMemory::Malloc(Size);
+	Dispatch->Uid = Uid;
 	Dispatch->FieldCount = FieldCount;
 	Dispatch->EventSize = 0;
 	Dispatch->FirstRoute = -1;
diff --git a/Engine/Source/Developer/TraceAnalysis/Public/Trace/Analyzer.h b/Engine/Source/Developer/TraceAnalysis/Public/Trace/Analyzer.h
--- a/Engine/Source/Developer/TraceAnalysis/Public/Trace/Analyzer.h
+++ b/Engine/Source/Developer/TraceAnalysis/Public/Trace/Analyzer.h
@@ -54,8 +54,49 @@ public:
 		FInterfaceBuilder&		InterfaceBuilder;
 	};
 
+	struct TRACEANALYSIS_API FEventFieldInfo
+	{
+		enum class EType { None, Integer, Float, };
+
+		/* Returns the name of the field. */
+		const ANSICHAR* GetName() const;
+
+		/* Offset from the start of the event to this field's data. */
+		uint32 GetOffset() const;
+
+		uint32 GetSize() const;
+
+		/* What type of field is this? */
+		EType GetType() const;
+	};
+
+	struct TRACEANALYSIS_API FEventTypeInfo
+	{
+		/* Each event is assigned a unique ID when logged. Not that this is not
+		 * guaranteed to be the same for the same event from one trace to the next. */
+		uint32 GetId() const;
+
+		/* Returns the total size of the event. */
+		uint32 GetSize() const;
+
+		/* The name of the event. */
+		const ANSICHAR* GetName() const;
+
+		/* Returns the logger name the event is associated with. */
+		const ANSICHAR* GetLoggerName() const;
+
+		/* The number of member fields this event has. */
+		uint32 GetFieldCount() const;
+
+		/* By-index access to fields' type information. */
+		const FEventFieldInfo* GetFieldInfo(uint32 Index) const;
+	};
+
 	struct TRACEANALYSIS_API FEventData
 	{
+		/** Returns an object describing the underlying event's type. */
+		const FEventTypeInfo& GetTypeInfo() const;
+
 		/** Queries the value of a field of the event. */
 		template <typename ValueType> ValueType GetValue(const ANSICHAR* FieldName) const;
 
