Granted access to the raw event data.

diff --git a/Engine/Source/Developer/TraceAnalysis/Private/Analysis/Engine.cpp b/Engine/Source/Developer/TraceAnalysis/Private/Analysis/Engine.cpp
--- a/Engine/Source/Developer/TraceAnalysis/Private/Analysis/Engine.cpp
+++ b/Engine/Source/Developer/TraceAnalysis/Private/Analysis/Engine.cpp
@@ -371,6 +371,13 @@ uint32 IAnalyzer::FEventData::GetAttachmentSize() const
 	return Info.Size - Info.Dispatch.EventSize;
 }
 
+////////////////////////////////////////////////////////////////////////////////
+const uint8* IAnalyzer::FEventData::GetRawPointer() const
+{
+	const auto& Info = *(const FAnalysisEngine::FEventDataInfo*)this;
+	return Info.Ptr;
+}
+
 
 
 ////////////////////////////////////////////////////////////////////////////////
diff --git a/Engine/Source/Developer/TraceAnalysis/Public/Trace/Analyzer.h b/Engine/Source/Developer/TraceAnalysis/Public/Trace/Analyzer.h
--- a/Engine/Source/Developer/TraceAnalysis/Public/Trace/Analyzer.h
+++ b/Engine/Source/Developer/TraceAnalysis/Public/Trace/Analyzer.h
@@ -108,6 +108,9 @@ public:
 		/** Returns the size of the events attachment, or 0 if none. */
 		uint32 GetAttachmentSize() const;
 
+		/** Provides a pointer to the raw event data. */
+		const uint8* GetRawPointer() const;
+
 	private:
 		const void* GetValueImpl(const ANSICHAR* FieldName, uint16& Type) const;
 	};
