Corrected my doxygen tags.

diff --git a/Engine/Source/Developer/TraceAnalysis/Public/Trace/Analyzer.h b/Engine/Source/Developer/TraceAnalysis/Public/Trace/Analyzer.h
--- a/Engine/Source/Developer/TraceAnalysis/Public/Trace/Analyzer.h
+++ b/Engine/Source/Developer/TraceAnalysis/Public/Trace/Analyzer.h
@@ -58,37 +58,37 @@ public:
 	{
 		enum class EType { None, Integer, Float, };
 
-		/* Returns the name of the field. */
+		/** Returns the name of the field. */
 		const ANSICHAR* GetName() const;
 
-		/* Offset from the start of the event to this field's data. */
+		/** Offset from the start of the event to this field's data. */
 		uint32 GetOffset() const;
 
 		uint32 GetSize() const;
 
-		/* What type of field is this? */
+		/** What type of field is this? */
 		EType GetType() const;
 	};
 
 	struct TRACEANALYSIS_API FEventTypeInfo
 	{
-		/* Each event is assigned a unique ID when logged. Not that this is not
+		/** Each event is assigned a unique ID when logged. Not that this is not
 		 * guaranteed to be the same for the same event from one trace to the next. */
 		uint32 GetId() const;
 
-		/* Returns the total size of the event. */
+		/** Returns the total size of the event. */
 		uint32 GetSize() const;
 
-		/* The name of the event. */
+		/** The name of the event. */
 		const ANSICHAR* GetName() const;
 
-		/* Returns the logger name the event is associated with. */
+		/** Returns the logger name the event is associated with. */
 		const ANSICHAR* GetLoggerName() const;
 
-		/* The number of member fields this event has. */
+		/** The number of member fields this event has. */
 		uint32 GetFieldCount() const;
 
-		/* By-index access to fields' type information. */
+		/** By-index access to fields' type information. */
 		const FEventFieldInfo* GetFieldInfo(uint32 Index) const;
 	};
 
