Instead of using a magic sentinal route at the end of the routes array, default the first route to something that will be out of bounds and simply bounds check.

diff --git a/Engine/Source/Developer/TraceAnalysis/Private/Analysis/Engine.cpp b/Engine/Source/Developer/TraceAnalysis/Private/Analysis/Engine.cpp
--- a/Engine/Source/Developer/TraceAnalysis/Private/Analysis/Engine.cpp
+++ b/Engine/Source/Developer/TraceAnalysis/Private/Analysis/Engine.cpp
@@ -503,11 +503,6 @@ void FAnalysisEngine::OnNewTrace(const FOnEventContext& Context)
 			Cursor->Count = 1;
 		}
 	}
-
-	// Add a terminal route for events that aren't subscribed to
-	FRoute& Route = Routes.Emplace_GetRef();
-	Route.Hash = ~0u;
-	Route.Count = 0;
 }
 
 ////////////////////////////////////////////////////////////////////////////////
@@ -529,7 +524,7 @@ FAnalysisEngine::FDispatch& FAnalysisEngine::AddDispatch(
 	Dispatch->Uid = Uid;
 	Dispatch->FieldCount = FieldCount;
 	Dispatch->EventSize = 0;
-	Dispatch->FirstRoute = -1;
+	Dispatch->FirstRoute = ~uint16(0);
 
 	// Add the new dispatch in the dispatch table
 	if (Uid >= Dispatches.Num())
@@ -623,7 +618,6 @@ void FAnalysisEngine::OnNewEvent(const FOnEventContext& Context)
 	}
 
 	// Find routes that have subscribed to this event.
-	Dispatch.FirstRoute = Routes.Num() - 1;
 	for (uint16 i = 0, n = Routes.Num(); i < n; ++i)
 	{
 		if (Routes[i].Hash == DispatchHash.Get())
@@ -736,11 +730,15 @@ bool FAnalysisEngine::OnData(FStreamReader::FData& Data)
 		FEventDataInfo EventDataInfo = { *Dispatch, Header->EventData, Header->Size };
 		const FEventData& EventData = (FEventData&)EventDataInfo;
 
-		const FRoute* Route = Routes.GetData() + Dispatch->FirstRoute;
-		for (uint32 n = Route->Count; n--; ++Route)
+		uint32 RouteCount = Routes.Num();
+		if (Dispatch->FirstRoute < RouteCount)
 		{
-			IAnalyzer* Analyzer = Analyzers[Route->AnalyzerIndex];
-			Analyzer->OnEvent(Route->Id, { SessionContext, EventData });
+			const FRoute* Route = Routes.GetData() + Dispatch->FirstRoute;
+			for (uint32 n = Route->Count; n--; ++Route)
+			{
+				IAnalyzer* Analyzer = Analyzers[Route->AnalyzerIndex];
+				Analyzer->OnEvent(Route->Id, { SessionContext, EventData });
+			}
 		}
 	}
 
