Strategic rename of OnNewEvent() to OnNewEventInternal().

diff --git a/Engine/Source/Developer/TraceAnalysis/Private/Analysis/Engine.cpp b/Engine/Source/Developer/TraceAnalysis/Private/Analysis/Engine.cpp
--- a/Engine/Source/Developer/TraceAnalysis/Private/Analysis/Engine.cpp
+++ b/Engine/Source/Developer/TraceAnalysis/Private/Analysis/Engine.cpp
@@ -430,7 +430,7 @@ void FAnalysisEngine::OnEvent(uint16 RouteId, const FOnEventContext& Context)
 {
 	switch (RouteId)
 	{
-	case RouteId_NewEvent:	return OnNewEvent(Context);
+	case RouteId_NewEvent:	return OnNewEventInternal(Context);
 	case RouteId_NewTrace:	return OnNewTrace(Context);
 	case RouteId_Timing:	return OnTiming(Context);
 	}
@@ -538,7 +538,7 @@ FAnalysisEngine::FDispatch& FAnalysisEngine::AddDispatch(
 }
 
 ////////////////////////////////////////////////////////////////////////////////
-void FAnalysisEngine::OnNewEvent(const FOnEventContext& Context)
+void FAnalysisEngine::OnNewEventInternal(const FOnEventContext& Context)
 {
 	const FEventDataInfo& EventData = (const FEventDataInfo&)(Context.EventData);
 	const auto& NewEvent = *(FNewEventEvent*)(EventData.Ptr);
diff --git a/Engine/Source/Developer/TraceAnalysis/Private/Analysis/Engine.h b/Engine/Source/Developer/TraceAnalysis/Private/Analysis/Engine.h
--- a/Engine/Source/Developer/TraceAnalysis/Private/Analysis/Engine.h
+++ b/Engine/Source/Developer/TraceAnalysis/Private/Analysis/Engine.h
@@ -45,7 +45,7 @@ private:
 	void				AddRoute(uint16 AnalyzerIndex, uint16 Id, uint32 Hash);
 	void				OnNewTrace(const FOnEventContext& Context);
 	void				OnTiming(const FOnEventContext& Context);
-	void				OnNewEvent(const FOnEventContext& Context);
+	void				OnNewEventInternal(const FOnEventContext& Context);
 	FSessionContext		SessionContext;
 	TArray<FRoute>		Routes;
 	TArray<IAnalyzer*>	Analyzers;
