Extended analyzer API to issue callbacks when a new event type is registered from the trace stream.

diff --git a/Engine/Source/Developer/TraceAnalysis/Private/Analysis/Engine.cpp b/Engine/Source/Developer/TraceAnalysis/Private/Analysis/Engine.cpp
--- a/Engine/Source/Developer/TraceAnalysis/Private/Analysis/Engine.cpp
+++ b/Engine/Source/Developer/TraceAnalysis/Private/Analysis/Engine.cpp
@@ -627,6 +627,17 @@ void FAnalysisEngine::OnNewEventInternal(const FOnEventContext& Context)
 		}
 	}
 
+	// Inform routes that a new event has been declared.
+	uint32 RouteCount = Routes.Num();
+	if (Dispatch.FirstRoute < RouteCount)
+	{
+		const FRoute* Route = Routes.GetData() + Dispatch.FirstRoute;
+		for (uint32 n = Route->Count; n--; ++Route)
+		{
+			IAnalyzer* Analyzer = Analyzers[Route->AnalyzerIndex];
+			Analyzer->OnNewEvent(Route->Id, (FEventTypeInfo&)Dispatch);
+		}
+	}
 }
 
 ////////////////////////////////////////////////////////////////////////////////
diff --git a/Engine/Source/Developer/TraceAnalysis/Public/Trace/Analyzer.h b/Engine/Source/Developer/TraceAnalysis/Public/Trace/Analyzer.h
--- a/Engine/Source/Developer/TraceAnalysis/Public/Trace/Analyzer.h
+++ b/Engine/Source/Developer/TraceAnalysis/Public/Trace/Analyzer.h
@@ -132,6 +132,12 @@ public:
 	 * further events */
 	virtual void OnAnalysisEnd() {}
 
+	/** When a new event type appears in the trace stream, this method is called
+	 * if the event type has been subscribed to.
+	 * @param RouteId User-provided identifier for this event subscription.
+	 * @param TypeInfo Object describing the new event's type. */
+	virtual void OnNewEvent(uint16 RouteId, const FEventTypeInfo& TypeInfo) {}
+
 	/** For each event subscribed to in OnAnalysisBegin(), the analysis engine
 	 * will call this method when those events are encountered in a trace log
 	 * @param RouteId User-provided identifier given when subscribing to a particular event
