Gracefully handle invalid event IDs.

diff --git a/Engine/Source/Developer/TraceAnalysis/Private/Analysis/Engine.cpp b/Engine/Source/Developer/TraceAnalysis/Private/Analysis/Engine.cpp
--- a/Engine/Source/Developer/TraceAnalysis/Private/Analysis/Engine.cpp
+++ b/Engine/Source/Developer/TraceAnalysis/Private/Analysis/Engine.cpp
@@ -737,6 +737,10 @@ bool FAnalysisEngine::OnData(FStreamReader::FData& Data)
 		Transport->Advance(BlockSize);
 
 		const FDispatch* Dispatch = Dispatches[Uid];
+		if (Dispatch == nullptr)
+		{
+			return false;
+		}
 
 		FEventDataInfo EventDataInfo = { *Dispatch, Header->EventData, Header->Size };
 		const FEventData& EventData = (FEventData&)EventDataInfo;
