Added missing documentation.

diff --git a/Engine/Source/Developer/TraceAnalysis/Public/Trace/Analyzer.h b/Engine/Source/Developer/TraceAnalysis/Public/Trace/Analyzer.h
--- a/Engine/Source/Developer/TraceAnalysis/Public/Trace/Analyzer.h
+++ b/Engine/Source/Developer/TraceAnalysis/Public/Trace/Analyzer.h
@@ -64,6 +64,7 @@ public:
 		/** Offset from the start of the event to this field's data. */
 		uint32 GetOffset() const;
 
+		/** Returns the size of the field's data. */
 		uint32 GetSize() const;
 
 		/** What type of field is this? */
