Plugged a memory leak.

(ushell-p4-stage of 8325528)

diff --git a/Engine/Source/Developer/TraceAnalysis/Private/Analysis/Engine.cpp b/Engine/Source/Developer/TraceAnalysis/Private/Analysis/Engine.cpp
--- a/Engine/Source/Developer/TraceAnalysis/Private/Analysis/Engine.cpp
+++ b/Engine/Source/Developer/TraceAnalysis/Private/Analysis/Engine.cpp
@@ -37,6 +37,7 @@ private:
 class FTransportReader
 {
 public:
+	virtual					~FTransportReader() {}
 	void					SetSource(FStreamReader::FData& InSource);
 	template <typename RetType>
 	RetType const*			GetPointer();
@@ -423,6 +424,8 @@ FAnalysisEngine::~FAnalysisEngine()
 	{
 		FMemory::Free(Dispatch);
 	}
+
+	delete Transport;
 }
 
 ////////////////////////////////////////////////////////////////////////////////
