Added option to subscribe to all of a trace stream's events

diff --git a/Engine/Source/Developer/TraceAnalysis/Private/Analysis/Engine.cpp b/Engine/Source/Developer/TraceAnalysis/Private/Analysis/Engine.cpp
--- a/Engine/Source/Developer/TraceAnalysis/Private/Analysis/Engine.cpp
+++ b/Engine/Source/Developer/TraceAnalysis/Private/Analysis/Engine.cpp
@@ -394,6 +394,7 @@ enum ERouteId : uint16
 enum EKnownRouteHashes : uint32
 {
 	RouteHash_NewEvent = 0, // must be 0 to match traces.
+	RouteHash_AllEvents,	
 };
 
 ////////////////////////////////////////////////////////////////////////////////
@@ -477,6 +478,11 @@ void FAnalysisEngine::OnNewTrace(const FOnEventContext& Context)
 			Self->AddRoute(AnalyzerIndex, RouteId, Logger, Event);
 		}
 
+		virtual void RouteAllEvents(uint16 RouteId) override
+		{
+			Self->AddRoute(AnalyzerIndex, RouteId, RouteHash_AllEvents);
+		}
+
 		FAnalysisEngine* Self;
 		uint16 AnalyzerIndex;
 	} Builder;
@@ -641,6 +647,16 @@ void FAnalysisEngine::OnNewEventInternal(const FOnEventContext& Context)
 			Analyzer->OnNewEvent(Route->Id, (FEventTypeInfo&)Dispatch);
 		}
 	}
+
+	const FRoute* Route = Routes.GetData() + 1;
+	if (RouteCount > 1 && Route->Hash == RouteHash_AllEvents)
+	{
+		for (uint32 n = Route->Count; n--; ++Route)
+		{
+			IAnalyzer* Analyzer = Analyzers[Route->AnalyzerIndex];
+			Analyzer->OnNewEvent(Route->Id, (FEventTypeInfo&)Dispatch);
+		}
+	}
 }
 
 ////////////////////////////////////////////////////////////////////////////////
@@ -758,6 +774,22 @@ bool FAnalysisEngine::OnData(FStreamReader::FData& Data)
 				Analyzer->OnEvent(Route->Id, { SessionContext, EventData });
 			}
 		}
+
+		// Don't broadcast internal events
+		if (!Uid) // TODO: This makes assumptions about EKnownEventUids::User. Instead we should add this information to the trace.
+		{
+			continue;
+		}
+
+		const FRoute* Route = Routes.GetData() + 1;
+		if (RouteCount > 1 && Route->Hash == RouteHash_AllEvents)
+		{
+			for (uint32 n = Route->Count; n--; ++Route)
+			{
+				IAnalyzer* Analyzer = Analyzers[Route->AnalyzerIndex];
+				Analyzer->OnEvent(Route->Id, { SessionContext, EventData });
+			}
+		}
 	}
 
 	return true;
diff --git a/Engine/Source/Developer/TraceAnalysis/Public/Trace/Analyzer.h b/Engine/Source/Developer/TraceAnalysis/Public/Trace/Analyzer.h
--- a/Engine/Source/Developer/TraceAnalysis/Public/Trace/Analyzer.h
+++ b/Engine/Source/Developer/TraceAnalysis/Public/Trace/Analyzer.h
@@ -46,6 +46,10 @@ public:
 		 * @param Logger Name of the logger that emits the event.
 		 * @param Event Name of the event to subscribe to. */
 		virtual void RouteEvent(uint16 RouteId, const ANSICHAR* Logger, const ANSICHAR* Event) = 0;
+
+		/** Subscribe to all events in the trace stream being analyzed.
+		 * @param RouteId User-provided identifier for this event subscription. */
+		virtual void RouteAllEvents(uint16 RouteId) = 0;
 	};
 
 	struct FOnAnalysisContext
