- Analyzers can now indicate that they've finished analysing.
- Updated the TraceServices module following return value change to IAnalyzer::On*Event()

diff --git a/Engine/Source/Developer/TraceAnalysis/Private/Analysis/Engine.cpp b/Engine/Source/Developer/TraceAnalysis/Private/Analysis/Engine.cpp
--- a/Engine/Source/Developer/TraceAnalysis/Private/Analysis/Engine.cpp
+++ b/Engine/Source/Developer/TraceAnalysis/Private/Analysis/Engine.cpp
@@ -554,7 +554,10 @@ FAnalysisEngine::~FAnalysisEngine()
 {
 	for (IAnalyzer* Analyzer : Analyzers)
 	{
-		Analyzer->OnAnalysisEnd();
+		if (Analyzer != nullptr)
+		{
+			Analyzer->OnAnalysisEnd();
+		}
 	}
 
 	for (FDispatch* Dispatch : Dispatches)
@@ -566,14 +569,35 @@ FAnalysisEngine::~FAnalysisEngine()
 }
 
 ////////////////////////////////////////////////////////////////////////////////
-void FAnalysisEngine::OnEvent(uint16 RouteId, const FOnEventContext& Context)
+void FAnalysisEngine::RetireAnalyzer(uint32 AnalyzerIndex)
+{
+	if (AnalyzerIndex >= uint32(Analyzers.Num()))
+	{
+		return;
+	}
+
+	IAnalyzer* Analyzer = Analyzers[AnalyzerIndex]; // this line is brought to you with the word "Analyzer" (mostly).
+	if (Analyzer == nullptr)
+	{
+		return;
+	}
+
+	Analyzer->OnAnalysisEnd();
+	Analyzers[AnalyzerIndex] = nullptr;
+}
+
+
+////////////////////////////////////////////////////////////////////////////////
+bool FAnalysisEngine::OnEvent(uint16 RouteId, const FOnEventContext& Context)
 {
 	switch (RouteId)
 	{
-	case RouteId_NewEvent:	return OnNewEventInternal(Context);
-	case RouteId_NewTrace:	return OnNewTrace(Context);
-	case RouteId_Timing:	return OnTiming(Context);
+	case RouteId_NewEvent:	OnNewEventInternal(Context);	break;
+	case RouteId_NewTrace:	OnNewTrace(Context);			break;
+	case RouteId_Timing:	OnTiming(Context);				break;
 	}
+
+	return true;
 }
 
 ////////////////////////////////////////////////////////////////////////////////
@@ -780,7 +804,15 @@ void FAnalysisEngine::OnNewEventInternal(const FOnEventContext& Context)
 		for (uint32 n = Route->Count; n--; ++Route)
 		{
 			IAnalyzer* Analyzer = Analyzers[Route->AnalyzerIndex];
-			Analyzer->OnNewEvent(Route->Id, (FEventTypeInfo&)Dispatch);
+			if (Analyzer == nullptr)
+			{
+				continue;
+			}
+
+			if (!Analyzer->OnNewEvent(Route->Id, (FEventTypeInfo&)Dispatch))
+			{
+				RetireAnalyzer(Route->AnalyzerIndex);
+			}
 		}
 	}
 
@@ -790,7 +822,15 @@ void FAnalysisEngine::OnNewEventInternal(const FOnEventContext& Context)
 		for (uint32 n = Route->Count; n--; ++Route)
 		{
 			IAnalyzer* Analyzer = Analyzers[Route->AnalyzerIndex];
-			Analyzer->OnNewEvent(Route->Id, (FEventTypeInfo&)Dispatch);
+			if (Analyzer == nullptr)
+			{
+				continue;
+			}
+
+			if (!Analyzer->OnNewEvent(Route->Id, (FEventTypeInfo&)Dispatch))
+			{
+				RetireAnalyzer(Route->AnalyzerIndex);
+			}
 		}
 	}
 }
@@ -908,7 +948,15 @@ bool FAnalysisEngine::OnData(FStreamReader::FData& Data)
 			for (uint32 n = Route->Count; n--; ++Route)
 			{
 				IAnalyzer* Analyzer = Analyzers[Route->AnalyzerIndex];
-				Analyzer->OnEvent(Route->Id, { SessionContext, EventData });
+				if (Analyzer == nullptr)
+				{
+					continue;
+				}
+
+				if (!Analyzer->OnEvent(Route->Id, { SessionContext, EventData }))
+				{
+					RetireAnalyzer(Route->AnalyzerIndex);
+				}
 			}
 		}
 
@@ -924,12 +972,30 @@ bool FAnalysisEngine::OnData(FStreamReader::FData& Data)
 			for (uint32 n = Route->Count; n--; ++Route)
 			{
 				IAnalyzer* Analyzer = Analyzers[Route->AnalyzerIndex];
-				Analyzer->OnEvent(Route->Id, { SessionContext, EventData });
+				if (Analyzer == nullptr)
+				{
+					continue;
+				}
+
+				if (!Analyzer->OnEvent(Route->Id, { SessionContext, EventData }))
+				{
+					RetireAnalyzer(Route->AnalyzerIndex);
+				}
 			}
 		}
 	}
 
-	return true;
+	// If there's no analyzers left we might as well not continue
+	int32 ActiveAnalyzerCount = 0;
+	for (IAnalyzer* Analyzer : Analyzers)
+	{
+		if ((Analyzer != nullptr) && (Analyzer != this))
+		{
+			ActiveAnalyzerCount++;
+		}
+	}
+
+	return (ActiveAnalyzerCount > 0);
 }
 
 } // namespace Trace
diff --git a/Engine/Source/Developer/TraceAnalysis/Private/Analysis/Engine.h b/Engine/Source/Developer/TraceAnalysis/Private/Analysis/Engine.h
--- a/Engine/Source/Developer/TraceAnalysis/Private/Analysis/Engine.h
+++ b/Engine/Source/Developer/TraceAnalysis/Private/Analysis/Engine.h
@@ -37,7 +37,7 @@ private:
 		uint16			_Unused0;
 	};
 
-	virtual void		OnEvent(uint16 RouteId, const FOnEventContext& Context) override;
+	virtual bool		OnEvent(uint16 RouteId, const FOnEventContext& Context) override;
 
 	bool				EstablishTransport(FStreamReader::FData& Data);
 	FDispatch&			AddDispatch(uint16 Uid, uint16 FieldCount=0, uint16 ExtraData=0);
@@ -46,6 +46,7 @@ private:
 	void				OnNewTrace(const FOnEventContext& Context);
 	void				OnTiming(const FOnEventContext& Context);
 	void				OnNewEventInternal(const FOnEventContext& Context);
+	void				RetireAnalyzer(uint32 AnalyzerIndex);
 	FSessionContext		SessionContext;
 	TArray<FRoute>		Routes;
 	TArray<IAnalyzer*>	Analyzers;
diff --git a/Engine/Source/Developer/TraceAnalysis/Public/Trace/Analyzer.h b/Engine/Source/Developer/TraceAnalysis/Public/Trace/Analyzer.h
--- a/Engine/Source/Developer/TraceAnalysis/Public/Trace/Analyzer.h
+++ b/Engine/Source/Developer/TraceAnalysis/Public/Trace/Analyzer.h
@@ -140,14 +140,22 @@ public:
 	/** When a new event type appears in the trace stream, this method is called
 	 * if the event type has been subscribed to.
 	 * @param RouteId User-provided identifier for this event subscription.
-	 * @param TypeInfo Object describing the new event's type. */
-	virtual void OnNewEvent(uint16 RouteId, const FEventTypeInfo& TypeInfo) {}
+	 * @param TypeInfo Object describing the new event's type.
+	 * @return This analyzer is removed from the analysis session if false is returned. */
+	virtual bool OnNewEvent(uint16 RouteId, const FEventTypeInfo& TypeInfo)
+	{
+		return true;
+	}
 
 	/** For each event subscribed to in OnAnalysisBegin(), the analysis engine
 	 * will call this method when those events are encountered in a trace log
-	 * @param RouteId User-provided identifier given when subscribing to a particular event
-	 * @param Context Access to the instance of the subscribed event */
-	virtual void OnEvent(uint16 RouteId, const FOnEventContext& Context) {}
+	 * @param RouteId User-provided identifier given when subscribing to a particular event.
+	 * @param Context Access to the instance of the subscribed event.
+	 * @return This analyzer is removed from the analysis session if false is returned. */
+	virtual bool OnEvent(uint16 RouteId, const FOnEventContext& Context)
+	{
+		return true;
+	}
 };
 
 ////////////////////////////////////////////////////////////////////////////////
diff --git a/Engine/Source/Developer/TraceServices/Private/Analyzers/CountersTraceAnalysis.cpp b/Engine/Source/Developer/TraceServices/Private/Analyzers/CountersTraceAnalysis.cpp
--- a/Engine/Source/Developer/TraceServices/Private/Analyzers/CountersTraceAnalysis.cpp
+++ b/Engine/Source/Developer/TraceServices/Private/Analyzers/CountersTraceAnalysis.cpp
@@ -19,7 +19,7 @@ void FCountersAnalyzer::OnAnalysisBegin(const FOnAnalysisContext& Context)
 	Builder.RouteEvent(RouteId_SetValueFloat, "Counters", "SetValueFloat");
 }
 
-void FCountersAnalyzer::OnEvent(uint16 RouteId, const FOnEventContext& Context)
+bool FCountersAnalyzer::OnEvent(uint16 RouteId, const FOnEventContext& Context)
 {
 	Trace::FAnalysisSessionEditScope _(Session);
 
@@ -69,4 +69,6 @@ void FCountersAnalyzer::OnEvent(uint16 RouteId, const FOnEventContext& Context)
 		break;
 	}
 	}
+
+	return true;
 }
diff --git a/Engine/Source/Developer/TraceServices/Private/Analyzers/CountersTraceAnalysis.h b/Engine/Source/Developer/TraceServices/Private/Analyzers/CountersTraceAnalysis.h
--- a/Engine/Source/Developer/TraceServices/Private/Analyzers/CountersTraceAnalysis.h
+++ b/Engine/Source/Developer/TraceServices/Private/Analyzers/CountersTraceAnalysis.h
@@ -19,7 +19,7 @@ class FCountersAnalyzer
 public:
 	FCountersAnalyzer(Trace::IAnalysisSession& Session, Trace::ICounterProvider& CounterProvider);
 	virtual void OnAnalysisBegin(const FOnAnalysisContext& Context) override;
-	virtual void OnEvent(uint16 RouteId, const FOnEventContext& Context) override;
+	virtual bool OnEvent(uint16 RouteId, const FOnEventContext& Context) override;
 
 private:
 	enum : uint16
diff --git a/Engine/Source/Developer/TraceServices/Private/Analyzers/CpuProfilerTraceAnalysis.cpp b/Engine/Source/Developer/TraceServices/Private/Analyzers/CpuProfilerTraceAnalysis.cpp
--- a/Engine/Source/Developer/TraceServices/Private/Analyzers/CpuProfilerTraceAnalysis.cpp
+++ b/Engine/Source/Developer/TraceServices/Private/Analyzers/CpuProfilerTraceAnalysis.cpp
@@ -20,7 +20,7 @@ void FCpuProfilerAnalyzer::OnAnalysisBegin(const FOnAnalysisContext& Context)
 	Builder.RouteEvent(RouteId_EndCapture, "CpuProfiler", "EndCapture");
 }
 
-void FCpuProfilerAnalyzer::OnEvent(uint16 RouteId, const FOnEventContext& Context)
+bool FCpuProfilerAnalyzer::OnEvent(uint16 RouteId, const FOnEventContext& Context)
 {
 	Trace::FAnalysisSessionEditScope _(Session);
 
@@ -100,6 +100,8 @@ void FCpuProfilerAnalyzer::OnEvent(uint16 RouteId, const FOnEventContext& Contex
 		break;
 	}
 	}
+
+	return true;
 }
 
 TSharedRef<FCpuProfilerAnalyzer::FThreadState> FCpuProfilerAnalyzer::GetThreadState(uint32 ThreadId)
diff --git a/Engine/Source/Developer/TraceServices/Private/Analyzers/CpuProfilerTraceAnalysis.h b/Engine/Source/Developer/TraceServices/Private/Analyzers/CpuProfilerTraceAnalysis.h
--- a/Engine/Source/Developer/TraceServices/Private/Analyzers/CpuProfilerTraceAnalysis.h
+++ b/Engine/Source/Developer/TraceServices/Private/Analyzers/CpuProfilerTraceAnalysis.h
@@ -18,7 +18,7 @@ class FCpuProfilerAnalyzer
 public:
 	FCpuProfilerAnalyzer(Trace::IAnalysisSession& Session, Trace::FTimingProfilerProvider& TimingProfilerProvider);
 	virtual void OnAnalysisBegin(const FOnAnalysisContext& Context) override;
-	virtual void OnEvent(uint16 RouteId, const FOnEventContext& Context) override;
+	virtual bool OnEvent(uint16 RouteId, const FOnEventContext& Context) override;
 
 private:
 	struct EventScopeState
diff --git a/Engine/Source/Developer/TraceServices/Private/Analyzers/CsvProfilerTraceAnalysis.cpp b/Engine/Source/Developer/TraceServices/Private/Analyzers/CsvProfilerTraceAnalysis.cpp
--- a/Engine/Source/Developer/TraceServices/Private/Analyzers/CsvProfilerTraceAnalysis.cpp
+++ b/Engine/Source/Developer/TraceServices/Private/Analyzers/CsvProfilerTraceAnalysis.cpp
@@ -59,7 +59,7 @@ void FCsvProfilerAnalyzer::OnAnalysisEnd()
 	ThreadStatesMap.Empty();
 }
 
-void FCsvProfilerAnalyzer::OnEvent(uint16 RouteId, const FOnEventContext& Context)
+bool FCsvProfilerAnalyzer::OnEvent(uint16 RouteId, const FOnEventContext& Context)
 {
 	Trace::FAnalysisSessionEditScope _(Session);
 
@@ -150,6 +150,8 @@ void FCsvProfilerAnalyzer::OnEvent(uint16 RouteId, const FOnEventContext& Contex
 		CsvProfilerProvider.EndCapture(CaptureEndFrame);
 	}
 	}
+
+	return true;
 }
 
 FCsvProfilerAnalyzer::FThreadState& FCsvProfilerAnalyzer::GetThreadState(uint32 ThreadId)
diff --git a/Engine/Source/Developer/TraceServices/Private/Analyzers/CsvProfilerTraceAnalysis.h b/Engine/Source/Developer/TraceServices/Private/Analyzers/CsvProfilerTraceAnalysis.h
--- a/Engine/Source/Developer/TraceServices/Private/Analyzers/CsvProfilerTraceAnalysis.h
+++ b/Engine/Source/Developer/TraceServices/Private/Analyzers/CsvProfilerTraceAnalysis.h
@@ -24,7 +24,7 @@ public:
 	FCsvProfilerAnalyzer(Trace::IAnalysisSession& Session, Trace::FCsvProfilerProvider& CsvProfilerProvider, Trace::ICounterProvider& CounterProvider, const Trace::IFrameProvider& FrameProvider, const Trace::IThreadProvider& ThreadProvider);
 	~FCsvProfilerAnalyzer();
 	virtual void OnAnalysisBegin(const FOnAnalysisContext& Context) override;
-	virtual void OnEvent(uint16 RouteId, const FOnEventContext& Context) override;
+	virtual bool OnEvent(uint16 RouteId, const FOnEventContext& Context) override;
 	virtual void OnAnalysisEnd() override;
 
 private:
diff --git a/Engine/Source/Developer/TraceServices/Private/Analyzers/GpuProfilerTraceAnalysis.cpp b/Engine/Source/Developer/TraceServices/Private/Analyzers/GpuProfilerTraceAnalysis.cpp
--- a/Engine/Source/Developer/TraceServices/Private/Analyzers/GpuProfilerTraceAnalysis.cpp
+++ b/Engine/Source/Developer/TraceServices/Private/Analyzers/GpuProfilerTraceAnalysis.cpp
@@ -20,7 +20,7 @@ void FGpuProfilerAnalyzer::OnAnalysisBegin(const FOnAnalysisContext& Context)
 	Builder.RouteEvent(RouteId_Frame, "GpuProfiler", "Frame");
 }
 
-void FGpuProfilerAnalyzer::OnEvent(uint16 RouteId, const FOnEventContext& Context)
+bool FGpuProfilerAnalyzer::OnEvent(uint16 RouteId, const FOnEventContext& Context)
 {
 	Trace::FAnalysisSessionEditScope _(Session);
 
@@ -75,6 +75,8 @@ void FGpuProfilerAnalyzer::OnEvent(uint16 RouteId, const FOnEventContext& Contex
 	}
 		
 	}
+
+	return true;
 }
 
 double FGpuProfilerAnalyzer::GpuTimestampToSessionTime(uint64 GpuMicroseconds)
diff --git a/Engine/Source/Developer/TraceServices/Private/Analyzers/GpuProfilerTraceAnalysis.h b/Engine/Source/Developer/TraceServices/Private/Analyzers/GpuProfilerTraceAnalysis.h
--- a/Engine/Source/Developer/TraceServices/Private/Analyzers/GpuProfilerTraceAnalysis.h
+++ b/Engine/Source/Developer/TraceServices/Private/Analyzers/GpuProfilerTraceAnalysis.h
@@ -17,7 +17,7 @@ class FGpuProfilerAnalyzer
 public:
 	FGpuProfilerAnalyzer(Trace::FAnalysisSession& Session, Trace::FTimingProfilerProvider& TimingProfilerProvider);
 	virtual void OnAnalysisBegin(const FOnAnalysisContext& Context) override;
-	virtual void OnEvent(uint16 RouteId, const FOnEventContext& Context) override;
+	virtual bool OnEvent(uint16 RouteId, const FOnEventContext& Context) override;
 
 private:
 	enum : uint16
diff --git a/Engine/Source/Developer/TraceServices/Private/Analyzers/LoadTimeTraceAnalysis.cpp b/Engine/Source/Developer/TraceServices/Private/Analyzers/LoadTimeTraceAnalysis.cpp
--- a/Engine/Source/Developer/TraceServices/Private/Analyzers/LoadTimeTraceAnalysis.cpp
+++ b/Engine/Source/Developer/TraceServices/Private/Analyzers/LoadTimeTraceAnalysis.cpp
@@ -89,7 +89,7 @@ void FAsyncLoadingTraceAnalyzer::OnAnalysisBegin(const FOnAnalysisContext& Conte
 	Builder.RouteEvent(RouteId_ClassInfo, "LoadTime", "ClassInfo");
 }
 
-void FAsyncLoadingTraceAnalyzer::OnEvent(uint16 RouteId, const FOnEventContext& Context)
+bool FAsyncLoadingTraceAnalyzer::OnEvent(uint16 RouteId, const FOnEventContext& Context)
 {
 	const auto& EventData = Context.EventData;
 	switch (RouteId)
@@ -350,6 +350,8 @@ void FAsyncLoadingTraceAnalyzer::OnEvent(uint16 RouteId, const FOnEventContext&
 		break;
 	}
 	}
+
+	return true;
 }
 
 void FAsyncLoadingTraceAnalyzer::FThreadState::EnterPackageScope(double Time, const Trace::FPackageInfo* PackageInfo, ELoadTimeProfilerPackageEventType EventType)
diff --git a/Engine/Source/Developer/TraceServices/Private/Analyzers/LoadTimeTraceAnalysis.h b/Engine/Source/Developer/TraceServices/Private/Analyzers/LoadTimeTraceAnalysis.h
--- a/Engine/Source/Developer/TraceServices/Private/Analyzers/LoadTimeTraceAnalysis.h
+++ b/Engine/Source/Developer/TraceServices/Private/Analyzers/LoadTimeTraceAnalysis.h
@@ -27,7 +27,7 @@ public:
 	virtual ~FAsyncLoadingTraceAnalyzer();
 
 	virtual void OnAnalysisBegin(const FOnAnalysisContext& Context) override;
-	virtual void OnEvent(uint16 RouteId, const FOnEventContext& Context) override;
+	virtual bool OnEvent(uint16 RouteId, const FOnEventContext& Context) override;
 	
 private:
 	struct FRequestState;
diff --git a/Engine/Source/Developer/TraceServices/Private/Analyzers/LogTraceAnalysis.cpp b/Engine/Source/Developer/TraceServices/Private/Analyzers/LogTraceAnalysis.cpp
--- a/Engine/Source/Developer/TraceServices/Private/Analyzers/LogTraceAnalysis.cpp
+++ b/Engine/Source/Developer/TraceServices/Private/Analyzers/LogTraceAnalysis.cpp
@@ -20,7 +20,7 @@ void FLogTraceAnalyzer::OnAnalysisBegin(const FOnAnalysisContext& Context)
 	Builder.RouteEvent(RouteId_LogMessage, "Logging", "LogMessage");
 }
 
-void FLogTraceAnalyzer::OnEvent(uint16 RouteId, const FOnEventContext& Context)
+bool FLogTraceAnalyzer::OnEvent(uint16 RouteId, const FOnEventContext& Context)
 {
 	Trace::FAnalysisSessionEditScope _(Session);
 
@@ -57,4 +57,6 @@ void FLogTraceAnalyzer::OnEvent(uint16 RouteId, const FOnEventContext& Context)
 		break;
 	}
 	}
+
+	return true;
 }
diff --git a/Engine/Source/Developer/TraceServices/Private/Analyzers/LogTraceAnalysis.h b/Engine/Source/Developer/TraceServices/Private/Analyzers/LogTraceAnalysis.h
--- a/Engine/Source/Developer/TraceServices/Private/Analyzers/LogTraceAnalysis.h
+++ b/Engine/Source/Developer/TraceServices/Private/Analyzers/LogTraceAnalysis.h
@@ -17,7 +17,7 @@ class FLogTraceAnalyzer
 public:
 	FLogTraceAnalyzer(Trace::IAnalysisSession& Session, Trace::FLogProvider& LogProvider);
 	virtual void OnAnalysisBegin(const FOnAnalysisContext& Context) override;
-	virtual void OnEvent(uint16 RouteId, const FOnEventContext& Context) override;
+	virtual bool OnEvent(uint16 RouteId, const FOnEventContext& Context) override;
 
 private:
 	enum : uint16
diff --git a/Engine/Source/Developer/TraceServices/Private/Analyzers/MiscTraceAnalysis.cpp b/Engine/Source/Developer/TraceServices/Private/Analyzers/MiscTraceAnalysis.cpp
--- a/Engine/Source/Developer/TraceServices/Private/Analyzers/MiscTraceAnalysis.cpp
+++ b/Engine/Source/Developer/TraceServices/Private/Analyzers/MiscTraceAnalysis.cpp
@@ -42,7 +42,7 @@ void FMiscTraceAnalyzer::OnAnalysisBegin(const FOnAnalysisContext& Context)
 	Builder.RouteEvent(RouteId_EndRenderFrame, "Misc", "EndRenderFrame");
 }
 
-void FMiscTraceAnalyzer::OnEvent(uint16 RouteId, const FOnEventContext& Context)
+bool FMiscTraceAnalyzer::OnEvent(uint16 RouteId, const FOnEventContext& Context)
 {
 	Trace::FAnalysisSessionEditScope _(Session);
 
@@ -161,6 +161,8 @@ void FMiscTraceAnalyzer::OnEvent(uint16 RouteId, const FOnEventContext& Context)
 		break;
 	}
 	}
+
+	return true;
 }
 
 FMiscTraceAnalyzer::FThreadState* FMiscTraceAnalyzer::GetThreadState(uint32 ThreadId)
diff --git a/Engine/Source/Developer/TraceServices/Private/Analyzers/MiscTraceAnalysis.h b/Engine/Source/Developer/TraceServices/Private/Analyzers/MiscTraceAnalysis.h
--- a/Engine/Source/Developer/TraceServices/Private/Analyzers/MiscTraceAnalysis.h
+++ b/Engine/Source/Developer/TraceServices/Private/Analyzers/MiscTraceAnalysis.h
@@ -27,7 +27,7 @@ public:
 					   Trace::FLogProvider& LogProvider,
 					   Trace::FFrameProvider& FrameProvider);
 	virtual void OnAnalysisBegin(const FOnAnalysisContext& Context) override;
-	virtual void OnEvent(uint16 RouteId, const FOnEventContext& Context) override;
+	virtual bool OnEvent(uint16 RouteId, const FOnEventContext& Context) override;
 
 private:
 	enum : uint16
diff --git a/Engine/Source/Developer/TraceServices/Private/Analyzers/NetTraceAnalyzer.cpp b/Engine/Source/Developer/TraceServices/Private/Analyzers/NetTraceAnalyzer.cpp
--- a/Engine/Source/Developer/TraceServices/Private/Analyzers/NetTraceAnalyzer.cpp
+++ b/Engine/Source/Developer/TraceServices/Private/Analyzers/NetTraceAnalyzer.cpp
@@ -35,14 +35,14 @@ void FNetTraceAnalyzer::OnAnalysisEnd()
 {
 }
 
-void FNetTraceAnalyzer::OnEvent(uint16 RouteId, const FOnEventContext& Context)
+bool FNetTraceAnalyzer::OnEvent(uint16 RouteId, const FOnEventContext& Context)
 {
 	Trace::FAnalysisSessionEditScope _(Session);
 
 	// check that we always get the InitEvent before processing any other events
 	if (!ensure(RouteId == RouteId_InitEvent || NetTraceVersion > 0))
 	{
-		return;
+		return false;
 	}
 
 	const auto& EventData = Context.EventData;
@@ -305,6 +305,8 @@ void FNetTraceAnalyzer::OnEvent(uint16 RouteId, const FOnEventContext& Context)
 		}
 		break;
 	}
+
+	return true;
 }
 
 TSharedRef<FNetTraceAnalyzer::FNetTraceGameInstanceState> FNetTraceAnalyzer::GetOrCreateActiveGameInstanceState(uint32 GameInstanceId)
diff --git a/Engine/Source/Developer/TraceServices/Private/Analyzers/NetTraceAnalyzer.h b/Engine/Source/Developer/TraceServices/Private/Analyzers/NetTraceAnalyzer.h
--- a/Engine/Source/Developer/TraceServices/Private/Analyzers/NetTraceAnalyzer.h
+++ b/Engine/Source/Developer/TraceServices/Private/Analyzers/NetTraceAnalyzer.h
@@ -18,7 +18,7 @@ class FNetTraceAnalyzer
 public:
 	FNetTraceAnalyzer(Trace::IAnalysisSession& Session, Trace::FNetProfilerProvider& NetProfilerProvider);
 	virtual void OnAnalysisBegin(const FOnAnalysisContext& Context) override;
-	virtual void OnEvent(uint16 RouteId, const FOnEventContext& Context) override;
+	virtual bool OnEvent(uint16 RouteId, const FOnEventContext& Context) override;
 	virtual void OnAnalysisEnd() override;
 
 private:
diff --git a/Engine/Source/Developer/TraceServices/Private/Analyzers/PlatformFileTraceAnalysis.cpp b/Engine/Source/Developer/TraceServices/Private/Analyzers/PlatformFileTraceAnalysis.cpp
--- a/Engine/Source/Developer/TraceServices/Private/Analyzers/PlatformFileTraceAnalysis.cpp
+++ b/Engine/Source/Developer/TraceServices/Private/Analyzers/PlatformFileTraceAnalysis.cpp
@@ -24,7 +24,7 @@ void FPlatformFileTraceAnalyzer::OnAnalysisBegin(const FOnAnalysisContext& Conte
 	Builder.RouteEvent(RouteId_EndWrite, "PlatformFile", "EndWrite");
 }
 
-void FPlatformFileTraceAnalyzer::OnEvent(uint16 RouteId, const FOnEventContext& Context)
+bool FPlatformFileTraceAnalyzer::OnEvent(uint16 RouteId, const FOnEventContext& Context)
 {
 	Trace::FAnalysisSessionEditScope _(Session);
 
@@ -170,4 +170,6 @@ void FPlatformFileTraceAnalyzer::OnEvent(uint16 RouteId, const FOnEventContext&
 		break;
 	}
 	}
+
+	return true;
 }
diff --git a/Engine/Source/Developer/TraceServices/Private/Analyzers/PlatformFileTraceAnalysis.h b/Engine/Source/Developer/TraceServices/Private/Analyzers/PlatformFileTraceAnalysis.h
--- a/Engine/Source/Developer/TraceServices/Private/Analyzers/PlatformFileTraceAnalysis.h
+++ b/Engine/Source/Developer/TraceServices/Private/Analyzers/PlatformFileTraceAnalysis.h
@@ -19,7 +19,7 @@ class FPlatformFileTraceAnalyzer
 public:
 	FPlatformFileTraceAnalyzer(Trace::IAnalysisSession& Session, Trace::FFileActivityProvider& FileActivityProvider);
 	virtual void OnAnalysisBegin(const FOnAnalysisContext& Context) override;
-	virtual void OnEvent(uint16 RouteId, const FOnEventContext& Context) override;
+	virtual bool OnEvent(uint16 RouteId, const FOnEventContext& Context) override;
 
 private:
 	enum : uint16
diff --git a/Engine/Source/Developer/TraceServices/Private/Analyzers/StatsTraceAnalysis.cpp b/Engine/Source/Developer/TraceServices/Private/Analyzers/StatsTraceAnalysis.cpp
--- a/Engine/Source/Developer/TraceServices/Private/Analyzers/StatsTraceAnalysis.cpp
+++ b/Engine/Source/Developer/TraceServices/Private/Analyzers/StatsTraceAnalysis.cpp
@@ -19,7 +19,7 @@ void FStatsAnalyzer::OnAnalysisBegin(const FOnAnalysisContext& Context)
 	Builder.RouteEvent(RouteId_EventBatch, "Stats", "EventBatch");
 }
 
-void FStatsAnalyzer::OnEvent(uint16 RouteId, const FOnEventContext& Context)
+bool FStatsAnalyzer::OnEvent(uint16 RouteId, const FOnEventContext& Context)
 {
 	Trace::FAnalysisSessionEditScope _(Session);
 
@@ -136,6 +136,8 @@ void FStatsAnalyzer::OnEvent(uint16 RouteId, const FOnEventContext& Context)
 		break;
 	}
 	}
+
+	return true;
 }
 
 TSharedRef<FStatsAnalyzer::FThreadState> FStatsAnalyzer::GetThreadState(uint32 ThreadId)
diff --git a/Engine/Source/Developer/TraceServices/Private/Analyzers/StatsTraceAnalysis.h b/Engine/Source/Developer/TraceServices/Private/Analyzers/StatsTraceAnalysis.h
--- a/Engine/Source/Developer/TraceServices/Private/Analyzers/StatsTraceAnalysis.h
+++ b/Engine/Source/Developer/TraceServices/Private/Analyzers/StatsTraceAnalysis.h
@@ -19,7 +19,7 @@ class FStatsAnalyzer
 public:
 	FStatsAnalyzer(Trace::IAnalysisSession& Session, Trace::ICounterProvider& CounterProvider);
 	virtual void OnAnalysisBegin(const FOnAnalysisContext& Context) override;
-	virtual void OnEvent(uint16 RouteId, const FOnEventContext& Context) override;
+	virtual bool OnEvent(uint16 RouteId, const FOnEventContext& Context) override;
 
 private:
 	enum : uint16
