Fixed FAnalysisProcessor::Impl memory leak

diff --git a/Engine/Source/Developer/TraceAnalysis/Private/Analysis/Processor.cpp b/Engine/Source/Developer/TraceAnalysis/Private/Analysis/Processor.cpp
--- a/Engine/Source/Developer/TraceAnalysis/Private/Analysis/Processor.cpp
+++ b/Engine/Source/Developer/TraceAnalysis/Private/Analysis/Processor.cpp
@@ -94,4 +94,16 @@ void FAnalysisProcessor::Stop()				{ if (Impl != nullptr) { Impl->StopAnalysis()
 void FAnalysisProcessor::Wait()				{ if (Impl != nullptr) { Impl->WaitOnAnalysis(); } }
 void FAnalysisProcessor::Pause(bool bState) { if (Impl != nullptr) { Impl->PauseAnalysis(bState); } }
 
+////////////////////////////////////////////////////////////////////////////////
+FAnalysisProcessor::FAnalysisProcessor(FAnalysisProcessor&& Rhs)
+{
+	Swap(Impl, Rhs.Impl);
+}
+
+////////////////////////////////////////////////////////////////////////////////
+FAnalysisProcessor::~FAnalysisProcessor()
+{
+	delete Impl;
+}
+
 } // namespace Trace
diff --git a/Engine/Source/Developer/TraceAnalysis/Public/Trace/Analysis.h b/Engine/Source/Developer/TraceAnalysis/Public/Trace/Analysis.h
--- a/Engine/Source/Developer/TraceAnalysis/Public/Trace/Analysis.h
+++ b/Engine/Source/Developer/TraceAnalysis/Public/Trace/Analysis.h
@@ -32,10 +32,17 @@ public:
 	 * @param bState Pause if true, resume if false. */
 	void Pause(bool bState);
 
+	~FAnalysisProcessor();
+	FAnalysisProcessor(FAnalysisProcessor&& Rhs);
+
 private:
-	friend		class FAnalysisContext;
-	class		FImpl;
-	FImpl*		Impl = nullptr;
+	friend				class FAnalysisContext;
+						FAnalysisProcessor() = default;
+						FAnalysisProcessor(FAnalysisProcessor&) = delete;
+	FAnalysisProcessor&	operator = (FAnalysisProcessor&) = delete;
+	FAnalysisProcessor&	operator = (FAnalysisProcessor&&) = delete;
+	class				FImpl;
+	FImpl*				Impl = nullptr;
 };
 
 
