Removed the commented out Lz4-compressed trace stream support.

diff --git a/Engine/Source/Developer/TraceAnalysis/Private/Analysis/Engine.cpp b/Engine/Source/Developer/TraceAnalysis/Private/Analysis/Engine.cpp
--- a/Engine/Source/Developer/TraceAnalysis/Private/Analysis/Engine.cpp
+++ b/Engine/Source/Developer/TraceAnalysis/Private/Analysis/Engine.cpp
@@ -6,9 +6,6 @@
 #include "HAL/UnrealMemory.h"
 #include "Trace/Analysis.h"
 #include "Trace/Analyzer.h"
-#if 0
-#include "Trace/Private/DataDecoder.h"
-#endif // 0
 #include "Trace/Private/EventDef.h"
 
 namespace Trace
@@ -220,135 +217,6 @@ bool FTlsTransportReader::GetNextBatch()
 
 
 
-////////////////////////////////////////////////////////////////////////////////
-#if 0
-class FLz4TransportReader
-	: public FTransportReader
-{
-	enum : uint32 { MaxChunks = 4 };
-
-public:
-							FLz4TransportReader(uint32 InChunkPow2);
-	virtual					~FLz4TransportReader();
-
-private:
-	virtual void			Advance(uint32 BlockSize) override;
-	virtual const uint8*	GetPointerImpl(uint32 BlockSize) override;
-	bool					NextChunk();
-    FDataDecoder			Decoder;
-	uint8*					Buffer;
-	const uint8*			Cursor;
-	const uint8*			End;
-	uint32					Index = 0;
-	const uint32			ChunkPow2;
-};
-
-////////////////////////////////////////////////////////////////////////////////
-FLz4TransportReader::FLz4TransportReader(uint32 InChunkPow2)
-: ChunkPow2(InChunkPow2)
-{
-	uint32 ChunkSize = 1 << ChunkPow2;
-	Buffer = new uint8[ChunkSize * (MaxChunks + 1)];
-	Cursor = End = Buffer + (ChunkSize * 2);
-}
-
-////////////////////////////////////////////////////////////////////////////////
-FLz4TransportReader::~FLz4TransportReader()
-{
-	delete[] Buffer;
-}
-
-////////////////////////////////////////////////////////////////////////////////
-void FLz4TransportReader::Advance(uint32 BlockSize)
-{
-	Cursor += BlockSize;
-}
-
-////////////////////////////////////////////////////////////////////////////////
-const uint8* FLz4TransportReader::GetPointerImpl(uint32 BlockSize)
-{
-	while (true)
-	{
-		uint32 Remaining = uint32(UPTRINT(End - Cursor));
-		if (Remaining >= BlockSize)
-		{
-			break;
-		}
-
-		if (!NextChunk())
-		{
-			return nullptr;
-		}
-	}
-
-	return Cursor;
-}
-
-////////////////////////////////////////////////////////////////////////////////
-bool FLz4TransportReader::NextChunk()
-{
-	const struct {
-		uint32	Size;
-		uint8	Data[];
-	}* Block;
-
-	Block = decltype(Block)(Source->GetPointer(sizeof(*Block)));
-	if (Block == nullptr)
-	{
-		return false;
-	}
-
-	int32 BlockSize = Block->Size;
-	if (BlockSize < 0)
-	{
-		BlockSize = ~BlockSize;
-	}
-
-	Block = decltype(Block)(Source->GetPointer(BlockSize));
-	if (Block == nullptr)
-	{
-		return false;
-	}
-
-	Index = (Index + 1) & (MaxChunks - 1);
-	if (!Index)
-	{
-		uint32 Remaining = uint32(UPTRINT(End - Cursor));
-		uint32 OverflowChunks = Remaining >> ChunkPow2;
-		if (OverflowChunks >= MaxChunks - 1)
-		{
-			return false;
-		}
-
-		char* Dest = (char*)(Buffer + (1ull << ChunkPow2));
-		Dest -= Remaining - OverflowChunks;
-		memcpy(Dest, Cursor, Remaining);
-		Cursor = (uint8*)Dest;
-		End = Cursor + Remaining;
-	}
-
-	uint32 SrcSize = BlockSize - sizeof(*Block);
-	uint32 DestSize = 1 << ChunkPow2;
-	int DecodedSize = Decoder.Decode(Block->Data, End, SrcSize, DestSize);
-	if (DecodedSize < 0)
-	{
-		return false;
-	}
-
-	End += DecodedSize;
-
-	if (BlockSize != Block->Size)
-	{
-		Decoder.Reset();
-	}
-
-	Source->Advance(BlockSize);
-	return true;
-}
-#endif // 0
-
-
-
 ////////////////////////////////////////////////////////////////////////////////
 struct FAnalysisEngine::FDispatch
 {
@@ -874,9 +742,6 @@ bool FAnalysisEngine::EstablishTransport(FStreamReader::FData& Data)
 	default:	return false;
 	//case 'E':	/* See the magic above */ break;
 	//case 'T':	/* See the magic above */ break;
-#if 0
-	case 4:		Transport = new FLz4TransportReader(Header->Parameter); break;
-#endif // 0
 	}
 
 	Data.Advance(sizeof(*Header));
