Moved the various different transport implementations out of Engine.cpp into their own units

diff --git a/Engine/Source/Developer/TraceAnalysis/Private/Analysis/Engine.cpp b/Engine/Source/Developer/TraceAnalysis/Private/Analysis/Engine.cpp
--- a/Engine/Source/Developer/TraceAnalysis/Private/Analysis/Engine.cpp
+++ b/Engine/Source/Developer/TraceAnalysis/Private/Analysis/Engine.cpp
@@ -7,6 +7,8 @@
 #include "Trace/Analysis.h"
 #include "Trace/Analyzer.h"
 #include "Trace/Private/EventDef.h"
+#include "Transport/Transport.h"
+#include "Transport/TlsTransport.h"
 
 namespace Trace
 {
@@ -30,193 +32,6 @@ private:
 
 
 
-////////////////////////////////////////////////////////////////////////////////
-class FTransportReader
-{
-public:
-	virtual					~FTransportReader() {}
-	void					SetSource(FStreamReader::FData& InSource);
-	template <typename RetType>
-	RetType const*			GetPointer();
-	template <typename RetType>
-	RetType const*			GetPointer(uint32 BlockSize);
-	virtual void			Advance(uint32 BlockSize);
-
-protected:
-	virtual const uint8*	GetPointerImpl(uint32 BlockSize);
-	FStreamReader::FData*	Source;
-};
-
-////////////////////////////////////////////////////////////////////////////////
-void FTransportReader::SetSource(FStreamReader::FData& InSource)
-{
-	Source = &InSource;
-}
-
-////////////////////////////////////////////////////////////////////////////////
-template <typename RetType>
-RetType const* FTransportReader::GetPointer()
-{
-	return GetPointer<RetType>(sizeof(RetType));
-}
-
-////////////////////////////////////////////////////////////////////////////////
-template <typename RetType>
-RetType const* FTransportReader::GetPointer(uint32 BlockSize)
-{
-	return (RetType const*)GetPointerImpl(BlockSize);
-}
-
-////////////////////////////////////////////////////////////////////////////////
-void FTransportReader::Advance(uint32 BlockSize)
-{
-	Source->Advance(BlockSize);
-}
-
-////////////////////////////////////////////////////////////////////////////////
-const uint8* FTransportReader::GetPointerImpl(uint32 BlockSize)
-{
-	return Source->GetPointer(BlockSize);
-}
-
-
-
-////////////////////////////////////////////////////////////////////////////////
-class FTlsTransportReader
-	: public FTransportReader
-{
-public:
-	virtual					~FTlsTransportReader();
-	virtual void			Advance(uint32 BlockSize) override;
-	virtual const uint8*	GetPointerImpl(uint32 BlockSize) override;
-
-private:
-	struct FPayloadNode
-	{
-		FPayloadNode*		Next;
-		uint32				Cursor;
-		uint16				Serial;
-		uint16				Size;
-		uint8				Data[];
-	};
-
-	bool					GetNextBatch();
-	static const uint32		MaxPayloadSize = 8192;
-	FPayloadNode*			ActiveList = nullptr;
-	FPayloadNode*			PendingList = nullptr;
-	FPayloadNode*			FreeList = nullptr;
-};
-
-////////////////////////////////////////////////////////////////////////////////
-FTlsTransportReader::~FTlsTransportReader()
-{
-	for (FPayloadNode* Root : {ActiveList, PendingList, FreeList})
-	{
-		for (FPayloadNode* Node = Root; Node != nullptr;)
-		{
-			FPayloadNode* Next = Node->Next;
-			delete[] Node;
-			Node = Next;
-		}
-	}
-}
-
-////////////////////////////////////////////////////////////////////////////////
-void FTlsTransportReader::Advance(uint32 BlockSize)
-{
-	if (ActiveList != nullptr)
-	{
-		ActiveList->Cursor += BlockSize;
-	}
-}
-
-////////////////////////////////////////////////////////////////////////////////
-const uint8* FTlsTransportReader::GetPointerImpl(uint32 BlockSize)
-{
-	if (ActiveList == nullptr && !GetNextBatch())
-	{
-		return nullptr;
-	}
-
-	uint32 NextCursor = ActiveList->Cursor + BlockSize;
-	if (NextCursor > ActiveList->Size)
-	{
-		FPayloadNode* Node = ActiveList;
-		ActiveList = ActiveList->Next;
-		Node->Next = FreeList;
-		FreeList = Node;
-		return GetPointerImpl(BlockSize);
-	}
-
-	return ActiveList->Data + ActiveList->Cursor;
-}
-
-////////////////////////////////////////////////////////////////////////////////
-bool FTlsTransportReader::GetNextBatch()
-{
-	struct FPayload
-	{
-		int16	Serial;
-		uint16	Size;
-		uint8	Data[];
-	};
-
-	int16 LastSerial = -1;
-	if (PendingList != nullptr)
-	{
-		LastSerial = PendingList->Serial;
-	}
-
-	while (true)
-	{
-		const auto* Payload = (const FPayload*)FTransportReader::GetPointerImpl(sizeof(FPayload));
-		if (Payload == nullptr)
-		{
-			return false;
-		}
-
-		// If this new payload is part of the next event batch then we've finished
-		// building the current batch. The current batch can be activated.
-		if (LastSerial >= Payload->Serial)
-		{
-			ActiveList = PendingList;
-			PendingList = nullptr;
-			break;
-		}
-
-		if (FTransportReader::GetPointerImpl(Payload->Size) == nullptr)
-		{
-			return false;
-		}
-
-		FTransportReader::Advance(Payload->Size);
-
-		FPayloadNode* Node;
-		if (FreeList != nullptr)
-		{
-			Node = FreeList;
-			FreeList = Node->Next;
-		}
-		else
-		{
-			Node = (FPayloadNode*)GMalloc->Malloc(sizeof(FPayloadNode) + MaxPayloadSize);
-		}
-
-		Node->Serial = Payload->Serial;
-		Node->Cursor = 0;
-		Node->Size = uint16(Payload->Size - sizeof(*Payload));
-		Node->Next = PendingList;
-		FMemory::Memcpy(Node->Data, Payload->Data, Node->Size);
-
-		PendingList = Node;
-		LastSerial = Payload->Serial;
-	}
-
-	return true;
-}
-
-
-
 ////////////////////////////////////////////////////////////////////////////////
 struct FAnalysisEngine::FDispatch
 {
@@ -454,7 +269,6 @@ void FAnalysisEngine::RetireAnalyzer(uint32 AnalyzerIndex)
 	Analyzers[AnalyzerIndex] = nullptr;
 }
 
-
 ////////////////////////////////////////////////////////////////////////////////
 bool FAnalysisEngine::OnEvent(uint16 RouteId, const FOnEventContext& Context)
 {
@@ -737,8 +551,8 @@ bool FAnalysisEngine::EstablishTransport(FStreamReader::FData& Data)
 
 	switch (Header->Format)
 	{
-	case 1:		Transport = new FTransportReader(); break;
-	case 2:		Transport = new FTlsTransportReader(); break;
+	case 1:		Transport = new FTransport(); break;
+	case 2:		Transport = new FTlsTransport(); break;
 	default:	return false;
 	//case 'E':	/* See the magic above */ break;
 	//case 'T':	/* See the magic above */ break;
diff --git a/Engine/Source/Developer/TraceAnalysis/Private/Analysis/Engine.h b/Engine/Source/Developer/TraceAnalysis/Private/Analysis/Engine.h
--- a/Engine/Source/Developer/TraceAnalysis/Private/Analysis/Engine.h
+++ b/Engine/Source/Developer/TraceAnalysis/Private/Analysis/Engine.h
@@ -14,7 +14,6 @@ namespace Trace
 
 ////////////////////////////////////////////////////////////////////////////////
 struct FNewEventEvent;
-class FTransportReader;
 
 ////////////////////////////////////////////////////////////////////////////////
 class FAnalysisEngine
@@ -51,7 +50,7 @@ private:
 	TArray<FRoute>		Routes;
 	TArray<IAnalyzer*>	Analyzers;
 	TArray<FDispatch*>	Dispatches;
-	FTransportReader*	Transport = nullptr;
+	class FTransport*	Transport = nullptr;
 };
 
 } // namespace Trace
diff --git a/Engine/Source/Developer/TraceAnalysis/Private/Analysis/Transport/TlsTransport.cpp b/Engine/Source/Developer/TraceAnalysis/Private/Analysis/Transport/TlsTransport.cpp
new file mode 100644
--- /dev/null
+++ b/Engine/Source/Developer/TraceAnalysis/Private/Analysis/Transport/TlsTransport.cpp
@@ -0,0 +1,116 @@
+// Copyright 1998-2019 Epic Games, Inc. All Rights Reserved.
+
+#include "TlsTransport.h"
+
+namespace Trace
+{
+
+////////////////////////////////////////////////////////////////////////////////
+FTlsTransport::~FTlsTransport()
+{
+	for (FPayloadNode* Root : {ActiveList, PendingList, FreeList})
+	{
+		for (FPayloadNode* Node = Root; Node != nullptr;)
+		{
+			FPayloadNode* Next = Node->Next;
+			delete[] Node;
+			Node = Next;
+		}
+	}
+}
+
+////////////////////////////////////////////////////////////////////////////////
+void FTlsTransport::Advance(uint32 BlockSize)
+{
+	if (ActiveList != nullptr)
+	{
+		ActiveList->Cursor += BlockSize;
+	}
+}
+
+////////////////////////////////////////////////////////////////////////////////
+const uint8* FTlsTransport::GetPointerImpl(uint32 BlockSize)
+{
+	if (ActiveList == nullptr && !GetNextBatch())
+	{
+		return nullptr;
+	}
+
+	uint32 NextCursor = ActiveList->Cursor + BlockSize;
+	if (NextCursor > ActiveList->Size)
+	{
+		FPayloadNode* Node = ActiveList;
+		ActiveList = ActiveList->Next;
+		Node->Next = FreeList;
+		FreeList = Node;
+		return GetPointerImpl(BlockSize);
+	}
+
+	return ActiveList->Data + ActiveList->Cursor;
+}
+
+////////////////////////////////////////////////////////////////////////////////
+bool FTlsTransport::GetNextBatch()
+{
+	struct FPayload
+	{
+		int16	Serial;
+		uint16	Size;
+		uint8	Data[];
+	};
+
+	int16 LastSerial = -1;
+	if (PendingList != nullptr)
+	{
+		LastSerial = PendingList->Serial;
+	}
+
+	while (true)
+	{
+		const auto* Payload = (const FPayload*)FTransport::GetPointerImpl(sizeof(FPayload));
+		if (Payload == nullptr)
+		{
+			return false;
+		}
+
+		// If this new payload is part of the next event batch then we've finished
+		// building the current batch. The current batch can be activated.
+		if (LastSerial >= Payload->Serial)
+		{
+			ActiveList = PendingList;
+			PendingList = nullptr;
+			break;
+		}
+
+		if (FTransport::GetPointerImpl(Payload->Size) == nullptr)
+		{
+			return false;
+		}
+
+		FTransport::Advance(Payload->Size);
+
+		FPayloadNode* Node;
+		if (FreeList != nullptr)
+		{
+			Node = FreeList;
+			FreeList = Node->Next;
+		}
+		else
+		{
+			Node = (FPayloadNode*)GMalloc->Malloc(sizeof(FPayloadNode) + MaxPayloadSize);
+		}
+
+		Node->Serial = Payload->Serial;
+		Node->Cursor = 0;
+		Node->Size = uint16(Payload->Size - sizeof(*Payload));
+		Node->Next = PendingList;
+		FMemory::Memcpy(Node->Data, Payload->Data, Node->Size);
+
+		PendingList = Node;
+		LastSerial = Payload->Serial;
+	}
+
+	return true;
+}
+
+} // namespace Trace
diff --git a/Engine/Source/Developer/TraceAnalysis/Private/Analysis/Transport/TlsTransport.h b/Engine/Source/Developer/TraceAnalysis/Private/Analysis/Transport/TlsTransport.h
new file mode 100644
--- /dev/null
+++ b/Engine/Source/Developer/TraceAnalysis/Private/Analysis/Transport/TlsTransport.h
@@ -0,0 +1,36 @@
+// Copyright 1998-2019 Epic Games, Inc. All Rights Reserved.
+
+#pragma once
+
+#include "Transport.h"
+
+namespace Trace
+{
+
+////////////////////////////////////////////////////////////////////////////////
+class FTlsTransport
+	: public FTransport
+{
+public:
+	virtual					~FTlsTransport();
+	virtual void			Advance(uint32 BlockSize) override;
+	virtual const uint8*	GetPointerImpl(uint32 BlockSize) override;
+
+private:
+	struct FPayloadNode
+	{
+		FPayloadNode*		Next;
+		uint32				Cursor;
+		uint16				Serial;
+		uint16				Size;
+		uint8				Data[];
+	};
+
+	bool					GetNextBatch();
+	static const uint32		MaxPayloadSize = 8192;
+	FPayloadNode*			ActiveList = nullptr;
+	FPayloadNode*			PendingList = nullptr;
+	FPayloadNode*			FreeList = nullptr;
+};
+
+} // namespace Trace
diff --git a/Engine/Source/Developer/TraceAnalysis/Private/Analysis/Transport/Transport.h b/Engine/Source/Developer/TraceAnalysis/Private/Analysis/Transport/Transport.h
new file mode 100644
--- /dev/null
+++ b/Engine/Source/Developer/TraceAnalysis/Private/Analysis/Transport/Transport.h
@@ -0,0 +1,59 @@
+// Copyright 1998-2019 Epic Games, Inc. All Rights Reserved.
+
+#pragma once
+
+#include "DataStream.h"
+
+namespace Trace
+{
+
+////////////////////////////////////////////////////////////////////////////////
+class FTransport
+{
+public:
+	virtual					~FTransport() {}
+	void					SetSource(FStreamReader::FData& InSource);
+	template <typename RetType>
+	RetType const*			GetPointer();
+	template <typename RetType>
+	RetType const*			GetPointer(uint32 BlockSize);
+	virtual void			Advance(uint32 BlockSize);
+
+protected:
+	virtual const uint8*	GetPointerImpl(uint32 BlockSize);
+	FStreamReader::FData*	Source;
+};
+
+////////////////////////////////////////////////////////////////////////////////
+inline void FTransport::SetSource(FStreamReader::FData& InSource)
+{
+	Source = &InSource;
+}
+
+////////////////////////////////////////////////////////////////////////////////
+template <typename RetType>
+inline RetType const* FTransport::GetPointer()
+{
+	return GetPointer<RetType>(sizeof(RetType));
+}
+
+////////////////////////////////////////////////////////////////////////////////
+template <typename RetType>
+inline RetType const* FTransport::GetPointer(uint32 BlockSize)
+{
+	return (RetType const*)GetPointerImpl(BlockSize);
+}
+
+////////////////////////////////////////////////////////////////////////////////
+inline void FTransport::Advance(uint32 BlockSize)
+{
+	Source->Advance(BlockSize);
+}
+
+////////////////////////////////////////////////////////////////////////////////
+inline const uint8* FTransport::GetPointerImpl(uint32 BlockSize)
+{
+	return Source->GetPointer(BlockSize);
+}
+
+} // namespace Trace
