Removed stale forward declarations

diff --git a/Engine/Source/Developer/TraceAnalysis/Private/Analysis/Engine.h b/Engine/Source/Developer/TraceAnalysis/Private/Analysis/Engine.h
--- a/Engine/Source/Developer/TraceAnalysis/Private/Analysis/Engine.h
+++ b/Engine/Source/Developer/TraceAnalysis/Private/Analysis/Engine.h
@@ -7,14 +7,9 @@
 #include "Trace/Analysis.h"
 #include "Trace/Analyzer.h"
 
-class IFileHandle;
-
 namespace Trace
 {
 
-////////////////////////////////////////////////////////////////////////////////
-struct FNewEventEvent;
-
 ////////////////////////////////////////////////////////////////////////////////
 class FAnalysisEngine
 	: public IAnalyzer
