Analysis can't rely on details such as constants and enums from the runtime as it should be version agnostic.

diff --git a/Engine/Source/Developer/TraceAnalysis/Private/Analysis/Engine.cpp b/Engine/Source/Developer/TraceAnalysis/Private/Analysis/Engine.cpp
--- a/Engine/Source/Developer/TraceAnalysis/Private/Analysis/Engine.cpp
+++ b/Engine/Source/Developer/TraceAnalysis/Private/Analysis/Engine.cpp
@@ -41,8 +41,7 @@ struct FAnalysisEngine::FDispatch
 		uint16		Offset;
 		uint16		Size;
 		uint16		NameOffset;			// From FField ptr
-		uint8		TypeInfo;
-		uint8		_Unused0;
+		int16		SizeAndType;		// value == byte_size, sign == float < 0 < int
 	};
 
 	uint16			Uid;
@@ -123,18 +122,23 @@ uint32 IAnalyzer::FEventFieldInfo::GetOffset() const
 uint32 IAnalyzer::FEventFieldInfo::GetSize() const
 {
 	const auto* Inner = (const FAnalysisEngine::FDispatch::FField*)this;
-	return 1 << (Inner->TypeInfo & _Field_Pow2SizeMask);
+	return (Inner->SizeAndType < 0) ? -(Inner->SizeAndType) : Inner->SizeAndType;
 }
 
 ////////////////////////////////////////////////////////////////////////////////
 IAnalyzer::FEventFieldInfo::EType IAnalyzer::FEventFieldInfo::GetType() const
 {
 	const auto* Inner = (const FAnalysisEngine::FDispatch::FField*)this;
-	switch (Inner->TypeInfo & _Field_CategoryMask)
+	if (Inner->SizeAndType > 0)
 	{
-	case _Field_Integer:	return EType::Integer;
-	case _Field_Float:		return EType::Float;
+		return EType::Integer;
 	}
+
+	if (Inner->SizeAndType < 0)
+	{
+		return EType::Float;
+	}
+
 	return EType::None;
 }
 
@@ -156,7 +160,7 @@ const IAnalyzer::FEventTypeInfo& IAnalyzer::FEventData::GetTypeInfo() const
 }
 
 ////////////////////////////////////////////////////////////////////////////////
-const void* IAnalyzer::FEventData::GetValueImpl(const ANSICHAR* FieldName, uint16& Type) const
+const void* IAnalyzer::FEventData::GetValueImpl(const ANSICHAR* FieldName, int16& SizeAndType) const
 {
 	const auto& Info = *(const FAnalysisEngine::FEventDataInfo*)this;
 
@@ -169,7 +173,7 @@ const void* IAnalyzer::FEventData::GetValueImpl(const ANSICHAR* FieldName, uint1
 		const auto& Field = Info.Dispatch.Fields[i];
 		if (Field.Hash == NameHash)
 		{
-			Type = Field.TypeInfo;
+			SizeAndType = Field.SizeAndType;
 			return (Info.Ptr + Field.Offset);
 		}
 	}
@@ -428,7 +432,12 @@ void FAnalysisEngine::OnNewEventInternal(const FOnEventContext& Context)
 		Out.Hash = FieldHash.Get();
 		Out.Offset = In.Offset;
 		Out.Size = In.Size;
-		Out.TypeInfo = In.TypeInfo;
+
+		Out.SizeAndType = 1 << (In.TypeInfo & Protocol0::Field_Pow2SizeMask);
+		if ((In.TypeInfo & Protocol0::Field_CategoryMask) == Protocol0::Field_Float)
+		{
+			Out.SizeAndType = -Out.SizeAndType;
+		}
 	}
 
 	// Write out names with null terminators.
diff --git a/Engine/Source/Developer/TraceAnalysis/Public/Trace/Analyzer.h b/Engine/Source/Developer/TraceAnalysis/Public/Trace/Analyzer.h
--- a/Engine/Source/Developer/TraceAnalysis/Public/Trace/Analyzer.h
+++ b/Engine/Source/Developer/TraceAnalysis/Public/Trace/Analyzer.h
@@ -117,7 +117,7 @@ public:
 		const uint8* GetRawPointer() const;
 
 	private:
-		const void* GetValueImpl(const ANSICHAR* FieldName, uint16& Type) const;
+		const void* GetValueImpl(const ANSICHAR* FieldName, int16& SizeAndType) const;
 	};
 
 	struct FOnEventContext
@@ -166,33 +166,21 @@ public:
 template <typename ValueType>
 ValueType IAnalyzer::FEventData::GetValue(const ANSICHAR* FieldName) const
 {
-	uint16 FieldTypeId;
-	const void* Addr = GetValueImpl(FieldName, FieldTypeId);
+	int16 FieldSizeAndType;
+	const void* Addr = GetValueImpl(FieldName, FieldSizeAndType);
 	if (Addr == nullptr)
 	{
 		return ValueType(0);
 	}
 
-	uint32 Pow2Size = (FieldTypeId & _Field_Pow2SizeMask);
-	uint32 Category = (FieldTypeId & _Field_CategoryMask);
-
-	if (Category == _Field_Float)
+	switch (FieldSizeAndType)
 	{
-		switch (Pow2Size)
-		{
-		case _Field_32: return ValueType(*(const float*)(Addr));
-		case _Field_64: return ValueType(*(const double*)(Addr));
-		}
-	}
-	else if (Category == _Field_Integer)
-	{
-		switch (Pow2Size)
-		{
-		case _Field_8:  return ValueType(*(const uint8*)(Addr));
-		case _Field_16: return ValueType(*(const uint16*)(Addr));
-		case _Field_32: return ValueType(*(const uint32*)(Addr));
-		case _Field_64: return ValueType(*(const uint64*)(Addr));
-		}
+	case -4: return ValueType(*(const float*)(Addr));
+	case -8: return ValueType(*(const double*)(Addr));
+	case  1: return ValueType(*(const uint8*)(Addr));
+	case  2: return ValueType(*(const uint16*)(Addr));
+	case  4: return ValueType(*(const uint32*)(Addr));
+	case  8: return ValueType(*(const uint64*)(Addr));
 	}
 
 	return ValueType(0);
