Protocol0 doesn't exist yet.

#rnx

diff --git a/Engine/Source/Developer/TraceAnalysis/Private/Analysis/Engine.cpp b/Engine/Source/Developer/TraceAnalysis/Private/Analysis/Engine.cpp
--- a/Engine/Source/Developer/TraceAnalysis/Private/Analysis/Engine.cpp
+++ b/Engine/Source/Developer/TraceAnalysis/Private/Analysis/Engine.cpp
@@ -433,8 +433,8 @@ void FAnalysisEngine::OnNewEventInternal(const FOnEventContext& Context)
 		Out.Offset = In.Offset;
 		Out.Size = In.Size;
 
-		Out.SizeAndType = 1 << (In.TypeInfo & Protocol0::Field_Pow2SizeMask);
-		if ((In.TypeInfo & Protocol0::Field_CategoryMask) == Protocol0::Field_Float)
+		Out.SizeAndType = 1 << (In.TypeInfo & _Field_Pow2SizeMask);
+		if ((In.TypeInfo & _Field_CategoryMask) == _Field_Float)
 		{
 			Out.SizeAndType = -Out.SizeAndType;
 		}
