Added support to tell Trace remotely to record to a file.

diff --git a/Engine/Source/Developer/TraceAnalysis/Private/ControlClient.cpp b/Engine/Source/Developer/TraceAnalysis/Private/ControlClient.cpp
--- a/Engine/Source/Developer/TraceAnalysis/Private/ControlClient.cpp
+++ b/Engine/Source/Developer/TraceAnalysis/Private/ControlClient.cpp
@@ -110,6 +110,17 @@ void FControlClient::SendSendTo(const TCHAR* Host)
     FormatAndSend(TEXT("SendTo %s"), Host);
 }
 
+////////////////////////////////////////////////////////////////////////////////
+void FControlClient::SendWriteTo(const TCHAR* Path)
+{
+    if (!IsConnected())
+    {
+        return;
+    }
+
+    FormatAndSend(TEXT("WriteTo %s"), Path);
+}
+
 ////////////////////////////////////////////////////////////////////////////////
 void FControlClient::SendToggleEvent(const TCHAR* EventSpec, bool bState)
 {
diff --git a/Engine/Source/Developer/TraceAnalysis/Public/Trace/ControlClient.h b/Engine/Source/Developer/TraceAnalysis/Public/Trace/ControlClient.h
--- a/Engine/Source/Developer/TraceAnalysis/Public/Trace/ControlClient.h
+++ b/Engine/Source/Developer/TraceAnalysis/Public/Trace/ControlClient.h
@@ -40,6 +40,10 @@ public:
 	 * @param Host Host IP address that the runtime should send events to. */
     void SendSendTo(const TCHAR* Host);
 
+	/** Tells the runtime to attempt to record its trace event log to a file.
+	 * @param Path File system path that the runtime should write events to. */
+    void SendWriteTo(const TCHAR* Path);
+
 	/** Toggles events on and off. Events can be specified in two forms; either
 	 * by logger to address a logger's events as a whole (EventSpec="LoggerName")
 	 * or individually (EventSpec="LoggerName.EventName").
