Tweaked method's comment to make units clearer

diff --git a/Engine/Source/Developer/TraceAnalysis/Public/Trace/Analyzer.h b/Engine/Source/Developer/TraceAnalysis/Public/Trace/Analyzer.h
--- a/Engine/Source/Developer/TraceAnalysis/Public/Trace/Analyzer.h
+++ b/Engine/Source/Developer/TraceAnalysis/Public/Trace/Analyzer.h
@@ -68,7 +68,7 @@ public:
 		/** Offset from the start of the event to this field's data. */
 		uint32 GetOffset() const;
 
-		/** Returns the size of the field's data. */
+		/** Returns the size of the field's data in bytes. */
 		uint32 GetSize() const;
 
 		/** What type of field is this? */
