Improved Trace analysis' tolerance against corrupt traces

#rnx

diff --git a/Engine/Source/Developer/TraceAnalysis/Private/Analysis/Engine.cpp b/Engine/Source/Developer/TraceAnalysis/Private/Analysis/Engine.cpp
--- a/Engine/Source/Developer/TraceAnalysis/Private/Analysis/Engine.cpp
+++ b/Engine/Source/Developer/TraceAnalysis/Private/Analysis/Engine.cpp
@@ -229,8 +229,8 @@ FAnalysisEngine::FAnalysisEngine(TArray<IAnalyzer*>&& InAnalyzers)
 
 	// Manually add event routing for known events, and those we don't quite know
 	// yet but are expecting.
-	FDispatch& NewEventDispatch = AddDispatch(uint16(FNewEventEvent::Uid), 0);
-	NewEventDispatch.FirstRoute = 0;
+	FDispatch* NewEventDispatch = AddDispatch(uint16(FNewEventEvent::Uid), 0);
+	NewEventDispatch->FirstRoute = 0;
 	AddRoute(SelfIndex, RouteId_NewEvent, RouteHash_NewEvent);
 	AddRoute(SelfIndex, RouteId_NewTrace, "$Trace", "NewTrace");
 	AddRoute(SelfIndex, RouteId_Timing, "$Trace", "Timing");
@@ -368,11 +368,25 @@ void FAnalysisEngine::OnTiming(const FOnEventContext& Context)
 }
 
 ////////////////////////////////////////////////////////////////////////////////
-FAnalysisEngine::FDispatch& FAnalysisEngine::AddDispatch(
+FAnalysisEngine::FDispatch* FAnalysisEngine::AddDispatch(
 	uint16 Uid,
 	uint16 FieldCount,
 	uint16 ExtraData)
 {
+	// Make sure there's enough space in the dispatch table, failing gently if
+	// there appears to be an existing entry.
+	if (Uid < Dispatches.Num())
+	{
+		if (Dispatches[Uid] != nullptr)
+		{
+			return nullptr;
+		}
+	}
+	else
+	{
+		Dispatches.SetNum(Uid + 1);
+	}
+
 	// Allocate a block of memory to hold the dispatch
 	uint32 Size = sizeof(FDispatch) + (sizeof(FDispatch::FField) * FieldCount) + ExtraData;
 	auto* Dispatch = (FDispatch*)FMemory::Malloc(Size);
@@ -382,14 +396,8 @@ FAnalysisEngine::FDispatch& FAnalysisEngine::AddDispatch(
 	Dispatch->FirstRoute = ~uint16(0);
 
 	// Add the new dispatch in the dispatch table
-	if (Uid >= Dispatches.Num())
-	{
-		Dispatches.SetNum(Uid + 1);
-	}
-	check(Dispatches[Uid] == nullptr);
 	Dispatches[Uid] = Dispatch;
-
-	return *Dispatch;
+	return Dispatch;
 }
 
 ////////////////////////////////////////////////////////////////////////////////
@@ -405,12 +413,17 @@ void FAnalysisEngine::OnNewEventInternal(const FOnEventContext& Context)
 	{
 		NameDataSize += NewEvent.Fields[i].NameSize + 1;
 	}
-	FDispatch& Dispatch = AddDispatch(NewEvent.EventUid, NewEvent.FieldCount, NameDataSize);
+
+	FDispatch* Dispatch = AddDispatch(NewEvent.EventUid, NewEvent.FieldCount, NameDataSize);
+	if (Dispatch == nullptr)
+	{
+		return;
+	}
 
 	if (NewEvent.FieldCount)
 	{
 		auto& LastField = NewEvent.Fields[NewEvent.FieldCount - 1];
-		Dispatch.EventSize = LastField.Offset + LastField.Size;
+		Dispatch->EventSize = LastField.Offset + LastField.Size;
 	}
 
 	const uint8* NameCursor = (const uint8*)(NewEvent.Fields + NewEvent.FieldCount);
@@ -424,7 +437,7 @@ void FAnalysisEngine::OnNewEventInternal(const FOnEventContext& Context)
 	for (int i = 0, n = NewEvent.FieldCount; i < n; ++i)
 	{
 		const auto& In = NewEvent.Fields[i];
-		auto& Out = Dispatch.Fields[i];
+		auto& Out = Dispatch->Fields[i];
 
 		FFnv1aHash FieldHash;
 		NameCursor = FieldHash.Add(NameCursor, In.NameSize);
@@ -442,7 +455,7 @@ void FAnalysisEngine::OnNewEventInternal(const FOnEventContext& Context)
 
 	// Write out names with null terminators.
 	NameCursor = (const uint8*)(NewEvent.Fields + NewEvent.FieldCount);
-	uint8* WriteCursor = (uint8*)(Dispatch.Fields + Dispatch.FieldCount);
+	uint8* WriteCursor = (uint8*)(Dispatch->Fields + Dispatch->FieldCount);
 	auto WriteName = [&] (uint32 Size)
 	{
 		memcpy(WriteCursor, NameCursor, Size);
@@ -451,30 +464,30 @@ void FAnalysisEngine::OnNewEventInternal(const FOnEventContext& Context)
 		WriteCursor[-1] = '\0';
 	};
 
-	UPTRINT EventNameOffset = UPTRINT(Dispatch.Fields + NewEvent.FieldCount);
-	EventNameOffset -= UPTRINT(&Dispatch);
+	UPTRINT EventNameOffset = UPTRINT(Dispatch->Fields + NewEvent.FieldCount);
+	EventNameOffset -= UPTRINT(Dispatch);
 	EventNameOffset += NewEvent.LoggerNameSize + 1;
-	Dispatch.EventNameOffset = uint16(EventNameOffset);
+	Dispatch->EventNameOffset = uint16(EventNameOffset);
 
 	WriteName(NewEvent.LoggerNameSize);
 	WriteName(NewEvent.EventNameSize);
 	for (int i = 0, n = NewEvent.FieldCount; i < n; ++i)
 	{
-		auto& Out = Dispatch.Fields[i];
-		Out.NameOffset = uint16(UPTRINT(WriteCursor) - UPTRINT(&Dispatch));
+		auto& Out = Dispatch->Fields[i];
+		Out.NameOffset = uint16(UPTRINT(WriteCursor) - UPTRINT(Dispatch));
 
 		WriteName(NewEvent.Fields[i].NameSize);
 	}
 
 	// Sort by hash so we can binary search when looking up.
-	TArrayView<FDispatch::FField> Fields(Dispatch.Fields, Dispatch.FieldCount);
+	TArrayView<FDispatch::FField> Fields(Dispatch->Fields, Dispatch->FieldCount);
 	Algo::SortBy(Fields, [] (const auto& Field) { return Field.Hash; });
 
 	// Fix up field name offsets
 	for (int i = 0, n = NewEvent.FieldCount; i < n; ++i)
 	{
-		auto& Out = Dispatch.Fields[i];
-		Out.NameOffset = uint16(UPTRINT(&Dispatch) + Out.NameOffset - UPTRINT(&Out));
+		auto& Out = Dispatch->Fields[i];
+		Out.NameOffset = uint16(UPTRINT(Dispatch) + Out.NameOffset - UPTRINT(&Out));
 	}
 
 	// Find routes that have subscribed to this event.
@@ -482,16 +495,16 @@ void FAnalysisEngine::OnNewEventInternal(const FOnEventContext& Context)
 	{
 		if (Routes[i].Hash == DispatchHash.Get())
 		{
-			Dispatch.FirstRoute = i;
+			Dispatch->FirstRoute = i;
 			break;
 		}
 	}
 
 	// Inform routes that a new event has been declared.
 	uint32 RouteCount = Routes.Num();
-	if (Dispatch.FirstRoute < RouteCount)
+	if (Dispatch->FirstRoute < RouteCount)
 	{
-		const FRoute* Route = Routes.GetData() + Dispatch.FirstRoute;
+		const FRoute* Route = Routes.GetData() + Dispatch->FirstRoute;
 		for (uint32 n = Route->Count; n--; ++Route)
 		{
 			IAnalyzer* Analyzer = Analyzers[Route->AnalyzerIndex];
@@ -500,7 +513,7 @@ void FAnalysisEngine::OnNewEventInternal(const FOnEventContext& Context)
 				continue;
 			}
 
-			if (!Analyzer->OnNewEvent(Route->Id, (FEventTypeInfo&)Dispatch))
+			if (!Analyzer->OnNewEvent(Route->Id, *(FEventTypeInfo*)Dispatch))
 			{
 				RetireAnalyzer(Route->AnalyzerIndex);
 			}
@@ -518,7 +531,7 @@ void FAnalysisEngine::OnNewEventInternal(const FOnEventContext& Context)
 				continue;
 			}
 
-			if (!Analyzer->OnNewEvent(Route->Id, (FEventTypeInfo&)Dispatch))
+			if (!Analyzer->OnNewEvent(Route->Id, *(FEventTypeInfo*)Dispatch))
 			{
 				RetireAnalyzer(Route->AnalyzerIndex);
 			}
diff --git a/Engine/Source/Developer/TraceAnalysis/Private/Analysis/Engine.h b/Engine/Source/Developer/TraceAnalysis/Private/Analysis/Engine.h
--- a/Engine/Source/Developer/TraceAnalysis/Private/Analysis/Engine.h
+++ b/Engine/Source/Developer/TraceAnalysis/Private/Analysis/Engine.h
@@ -34,7 +34,7 @@ private:
 	virtual bool		OnEvent(uint16 RouteId, const FOnEventContext& Context) override;
 
 	bool				EstablishTransport(FStreamReader::FData& Data);
-	FDispatch&			AddDispatch(uint16 Uid, uint16 FieldCount=0, uint16 ExtraData=0);
+	FDispatch*			AddDispatch(uint16 Uid, uint16 FieldCount=0, uint16 ExtraData=0);
 	void				AddRoute(uint16 AnalyzerIndex, uint16 Id, const ANSICHAR* Logger, const ANSICHAR* Event);
 	void				AddRoute(uint16 AnalyzerIndex, uint16 Id, uint32 Hash);
 	void				OnNewTrace(const FOnEventContext& Context);
