Retire analyzers by pointer instead of analyzer index as that will be more convenient in the future.

#rnx

diff --git a/Engine/Source/Developer/TraceAnalysis/Private/Analysis/Engine.cpp b/Engine/Source/Developer/TraceAnalysis/Private/Analysis/Engine.cpp
--- a/Engine/Source/Developer/TraceAnalysis/Private/Analysis/Engine.cpp
+++ b/Engine/Source/Developer/TraceAnalysis/Private/Analysis/Engine.cpp
@@ -256,21 +256,19 @@ FAnalysisEngine::~FAnalysisEngine()
 }
 
 ////////////////////////////////////////////////////////////////////////////////
-void FAnalysisEngine::RetireAnalyzer(uint32 AnalyzerIndex)
+void FAnalysisEngine::RetireAnalyzer(IAnalyzer* Analyzer)
 {
-	if (AnalyzerIndex >= uint32(Analyzers.Num()))
+	for (uint32 i = 0, n = Analyzers.Num(); i < n; ++i)
 	{
-		return;
-	}
+		if (Analyzers[i] != Analyzer)
+		{
+			continue;
+		}
 
-	IAnalyzer* Analyzer = Analyzers[AnalyzerIndex]; // this line is brought to you with the word "Analyzer" (mostly).
-	if (Analyzer == nullptr)
-	{
-		return;
+		Analyzer->OnAnalysisEnd();
+		Analyzers[i] = nullptr;
+		break;
 	}
-
-	Analyzer->OnAnalysisEnd();
-	Analyzers[AnalyzerIndex] = nullptr;
 }
 
 ////////////////////////////////////////////////////////////////////////////////
@@ -515,7 +513,7 @@ void FAnalysisEngine::OnNewEventInternal(const FOnEventContext& Context)
 
 			if (!Analyzer->OnNewEvent(Route->Id, *(FEventTypeInfo*)Dispatch))
 			{
-				RetireAnalyzer(Route->AnalyzerIndex);
+				RetireAnalyzer(Analyzer);
 			}
 		}
 	}
@@ -533,7 +531,7 @@ void FAnalysisEngine::OnNewEventInternal(const FOnEventContext& Context)
 
 			if (!Analyzer->OnNewEvent(Route->Id, *(FEventTypeInfo*)Dispatch))
 			{
-				RetireAnalyzer(Route->AnalyzerIndex);
+				RetireAnalyzer(Analyzer);
 			}
 		}
 	}
@@ -656,7 +654,7 @@ bool FAnalysisEngine::OnData(FStreamReader::FData& Data)
 
 				if (!Analyzer->OnEvent(Route->Id, { SessionContext, EventData }))
 				{
-					RetireAnalyzer(Route->AnalyzerIndex);
+					RetireAnalyzer(Analyzer);
 				}
 			}
 		}
@@ -680,7 +678,7 @@ bool FAnalysisEngine::OnData(FStreamReader::FData& Data)
 
 				if (!Analyzer->OnEvent(Route->Id, { SessionContext, EventData }))
 				{
-					RetireAnalyzer(Route->AnalyzerIndex);
+					RetireAnalyzer(Analyzer);
 				}
 			}
 		}
diff --git a/Engine/Source/Developer/TraceAnalysis/Private/Analysis/Engine.h b/Engine/Source/Developer/TraceAnalysis/Private/Analysis/Engine.h
--- a/Engine/Source/Developer/TraceAnalysis/Private/Analysis/Engine.h
+++ b/Engine/Source/Developer/TraceAnalysis/Private/Analysis/Engine.h
@@ -40,7 +40,7 @@ private:
 	void				OnNewTrace(const FOnEventContext& Context);
 	void				OnTiming(const FOnEventContext& Context);
 	void				OnNewEventInternal(const FOnEventContext& Context);
-	void				RetireAnalyzer(uint32 AnalyzerIndex);
+	void				RetireAnalyzer(IAnalyzer* Analyzer);
 	FSessionContext		SessionContext;
 	TArray<FRoute>		Routes;
 	TArray<IAnalyzer*>	Analyzers;
