Use to pointers instead of references when getting to the underlying object
#rnx

diff --git a/Engine/Source/Developer/TraceAnalysis/Private/Analysis/Engine.cpp b/Engine/Source/Developer/TraceAnalysis/Private/Analysis/Engine.cpp
--- a/Engine/Source/Developer/TraceAnalysis/Private/Analysis/Engine.cpp
+++ b/Engine/Source/Developer/TraceAnalysis/Private/Analysis/Engine.cpp
@@ -58,36 +58,36 @@ struct FAnalysisEngine::FDispatch
 ////////////////////////////////////////////////////////////////////////////////
 uint32 IAnalyzer::FEventTypeInfo::GetId() const
 {
-	const auto& Inner = *(const FAnalysisEngine::FDispatch*)this;
-	return Inner.Uid;
+	const auto* Inner = (const FAnalysisEngine::FDispatch*)this;
+	return Inner->Uid;
 }
 
 ////////////////////////////////////////////////////////////////////////////////
 uint32 IAnalyzer::FEventTypeInfo::GetSize() const
 {
-	const auto& Inner = *(const FAnalysisEngine::FDispatch*)this;
-	return Inner.EventSize;
+	const auto* Inner = (const FAnalysisEngine::FDispatch*)this;
+	return Inner->EventSize;
 }
 
 ////////////////////////////////////////////////////////////////////////////////
 const ANSICHAR* IAnalyzer::FEventTypeInfo::GetName() const
 {
-	const auto& Inner = *(const FAnalysisEngine::FDispatch*)this;
-	return (const ANSICHAR*)(UPTRINT(&Inner) + Inner.EventNameOffset);
+	const auto* Inner = (const FAnalysisEngine::FDispatch*)this;
+	return (const ANSICHAR*)Inner + Inner->EventNameOffset;
 }
 
 ////////////////////////////////////////////////////////////////////////////////
 const ANSICHAR* IAnalyzer::FEventTypeInfo::GetLoggerName() const
 {
-	const auto& Inner = *(const FAnalysisEngine::FDispatch*)this;
-	return (const ANSICHAR*)(Inner.Fields + Inner.FieldCount);
+	const auto* Inner = (const FAnalysisEngine::FDispatch*)this;
+	return (const ANSICHAR*)(Inner->Fields + Inner->FieldCount);
 }
 
 ////////////////////////////////////////////////////////////////////////////////
 uint32 IAnalyzer::FEventTypeInfo::GetFieldCount() const
 {
-	const auto& Inner = *(const FAnalysisEngine::FDispatch*)this;
-	return Inner.FieldCount;
+	const auto* Inner = (const FAnalysisEngine::FDispatch*)this;
+	return Inner->FieldCount;
 }
 
 ////////////////////////////////////////////////////////////////////////////////
@@ -98,8 +98,8 @@ const IAnalyzer::FEventFieldInfo* IAnalyzer::FEventTypeInfo::GetFieldInfo(uint32
 		return nullptr;
 	}
 
-	const auto& Inner = *(const FAnalysisEngine::FDispatch*)this;
-	return (const IAnalyzer::FEventFieldInfo*)(Inner.Fields + Index);
+	const auto* Inner = (const FAnalysisEngine::FDispatch*)this;
+	return (const IAnalyzer::FEventFieldInfo*)(Inner->Fields + Index);
 }
 
 
@@ -155,26 +155,26 @@ struct FAnalysisEngine::FEventDataInfo
 ////////////////////////////////////////////////////////////////////////////////
 const IAnalyzer::FEventTypeInfo& IAnalyzer::FEventData::GetTypeInfo() const
 {
-	const auto& Info = *(const FAnalysisEngine::FEventDataInfo*)this;
-	return (const FEventTypeInfo&)(Info.Dispatch);
+	const auto* Info = (const FAnalysisEngine::FEventDataInfo*)this;
+	return (const FEventTypeInfo&)(Info->Dispatch);
 }
 
 ////////////////////////////////////////////////////////////////////////////////
 const void* IAnalyzer::FEventData::GetValueImpl(const ANSICHAR* FieldName, int16& SizeAndType) const
 {
-	const auto& Info = *(const FAnalysisEngine::FEventDataInfo*)this;
+	const auto* Info = (const FAnalysisEngine::FEventDataInfo*)this;
 
 	FFnv1aHash Hash;
 	Hash.Add(FieldName);
 	uint32 NameHash = Hash.Get();
 
-	for (int i = 0, n = Info.Dispatch.FieldCount; i < n; ++i)
+	for (int i = 0, n = Info->Dispatch.FieldCount; i < n; ++i)
 	{
-		const auto& Field = Info.Dispatch.Fields[i];
+		const auto& Field = Info->Dispatch.Fields[i];
 		if (Field.Hash == NameHash)
 		{
 			SizeAndType = Field.SizeAndType;
-			return (Info.Ptr + Field.Offset);
+			return (Info->Ptr + Field.Offset);
 		}
 	}
 
@@ -184,22 +184,22 @@ const void* IAnalyzer::FEventData::GetValueImpl(const ANSICHAR* FieldName, int16
 ////////////////////////////////////////////////////////////////////////////////
 const uint8* IAnalyzer::FEventData::GetAttachment() const
 {
-	const auto& Info = *(const FAnalysisEngine::FEventDataInfo*)this;
-	return Info.Ptr + Info.Dispatch.EventSize;
+	const auto* Info = (const FAnalysisEngine::FEventDataInfo*)this;
+	return Info->Ptr + Info->Dispatch.EventSize;
 }
 
 ////////////////////////////////////////////////////////////////////////////////
 uint32 IAnalyzer::FEventData::GetAttachmentSize() const
 {
-	const auto& Info = *(const FAnalysisEngine::FEventDataInfo*)this;
-	return Info.Size - Info.Dispatch.EventSize;
+	const auto* Info = (const FAnalysisEngine::FEventDataInfo*)this;
+	return Info->Size - Info->Dispatch.EventSize;
 }
 
 ////////////////////////////////////////////////////////////////////////////////
 const uint8* IAnalyzer::FEventData::GetRawPointer() const
 {
-	const auto& Info = *(const FAnalysisEngine::FEventDataInfo*)this;
-	return Info.Ptr;
+	const auto* Info = (const FAnalysisEngine::FEventDataInfo*)this;
+	return Info->Ptr;
 }
 
 
