Consolidated the dispatching of trace events to analyzers into one place to reduce redundancy and to make it easier to iterate on.

#rnx

diff --git a/Engine/Source/Developer/TraceAnalysis/Private/Analysis/Engine.cpp b/Engine/Source/Developer/TraceAnalysis/Private/Analysis/Engine.cpp
--- a/Engine/Source/Developer/TraceAnalysis/Private/Analysis/Engine.cpp
+++ b/Engine/Source/Developer/TraceAnalysis/Private/Analysis/Engine.cpp
@@ -398,6 +398,38 @@ FAnalysisEngine::FDispatch* FAnalysisEngine::AddDispatch(
 	return Dispatch;
 }
 
+////////////////////////////////////////////////////////////////////////////////
+template <typename ImplType>
+void FAnalysisEngine::ForEachRoute(const FDispatch* Dispatch, ImplType&& Impl)
+{
+	uint32 RouteCount = Routes.Num();
+	if (Dispatch->FirstRoute < RouteCount)
+	{
+		const FRoute* Route = Routes.GetData() + Dispatch->FirstRoute;
+		for (uint32 n = Route->Count; n--; ++Route)
+		{
+			IAnalyzer* Analyzer = Analyzers[Route->AnalyzerIndex];
+			if (Analyzer != nullptr)
+			{
+				Impl(Analyzer, Route->Id);
+			}
+		}
+	}
+
+	const FRoute* Route = Routes.GetData() + 1;
+	if (RouteCount > 1 && Route->Hash == RouteHash_AllEvents)
+	{
+		for (uint32 n = Route->Count; n--; ++Route)
+		{
+			IAnalyzer* Analyzer = Analyzers[Route->AnalyzerIndex];
+			if (Analyzer != nullptr)
+			{
+				Impl(Analyzer, Route->Id);
+			}
+		}
+	}
+}
+
 ////////////////////////////////////////////////////////////////////////////////
 void FAnalysisEngine::OnNewEventInternal(const FOnEventContext& Context)
 {
@@ -499,42 +531,13 @@ void FAnalysisEngine::OnNewEventInternal(const FOnEventContext& Context)
 	}
 
 	// Inform routes that a new event has been declared.
-	uint32 RouteCount = Routes.Num();
-	if (Dispatch->FirstRoute < RouteCount)
-	{
-		const FRoute* Route = Routes.GetData() + Dispatch->FirstRoute;
-		for (uint32 n = Route->Count; n--; ++Route)
-		{
-			IAnalyzer* Analyzer = Analyzers[Route->AnalyzerIndex];
-			if (Analyzer == nullptr)
-			{
-				continue;
-			}
-
-			if (!Analyzer->OnNewEvent(Route->Id, *(FEventTypeInfo*)Dispatch))
-			{
-				RetireAnalyzer(Analyzer);
-			}
-		}
-	}
-
-	const FRoute* Route = Routes.GetData() + 1;
-	if (RouteCount > 1 && Route->Hash == RouteHash_AllEvents)
+	ForEachRoute(Dispatch, [&] (IAnalyzer* Analyzer, uint16 RouteId)
 	{
-		for (uint32 n = Route->Count; n--; ++Route)
+		if (!Analyzer->OnNewEvent(RouteId, *(FEventTypeInfo*)Dispatch))
 		{
-			IAnalyzer* Analyzer = Analyzers[Route->AnalyzerIndex];
-			if (Analyzer == nullptr)
-			{
-				continue;
-			}
-
-			if (!Analyzer->OnNewEvent(Route->Id, *(FEventTypeInfo*)Dispatch))
-			{
-				RetireAnalyzer(Analyzer);
-			}
+			RetireAnalyzer(Analyzer);
 		}
-	}
+	});
 }
 
 ////////////////////////////////////////////////////////////////////////////////
@@ -640,48 +643,13 @@ bool FAnalysisEngine::OnData(FStreamReader::FData& Data)
 		FEventDataInfo EventDataInfo = { *Dispatch, Header->EventData, Header->Size };
 		const FEventData& EventData = (FEventData&)EventDataInfo;
 
-		uint32 RouteCount = Routes.Num();
-		if (Dispatch->FirstRoute < RouteCount)
-		{
-			const FRoute* Route = Routes.GetData() + Dispatch->FirstRoute;
-			for (uint32 n = Route->Count; n--; ++Route)
-			{
-				IAnalyzer* Analyzer = Analyzers[Route->AnalyzerIndex];
-				if (Analyzer == nullptr)
-				{
-					continue;
-				}
-
-				if (!Analyzer->OnEvent(Route->Id, { SessionContext, EventData }))
-				{
-					RetireAnalyzer(Analyzer);
-				}
-			}
-		}
-
-		// Don't broadcast internal events
-		if (!Uid) // TODO: This makes assumptions about EKnownEventUids::User. Instead we should add this information to the trace.
-		{
-			continue;
-		}
-
-		const FRoute* Route = Routes.GetData() + 1;
-		if (RouteCount > 1 && Route->Hash == RouteHash_AllEvents)
+		ForEachRoute(Dispatch, [&] (IAnalyzer* Analyzer, uint16 RouteId)
 		{
-			for (uint32 n = Route->Count; n--; ++Route)
+			if (!Analyzer->OnEvent(RouteId, { SessionContext, EventData }))
 			{
-				IAnalyzer* Analyzer = Analyzers[Route->AnalyzerIndex];
-				if (Analyzer == nullptr)
-				{
-					continue;
-				}
-
-				if (!Analyzer->OnEvent(Route->Id, { SessionContext, EventData }))
-				{
-					RetireAnalyzer(Analyzer);
-				}
+				RetireAnalyzer(Analyzer);
 			}
-		}
+		});
 	}
 
 	// If there's no analyzers left we might as well not continue
diff --git a/Engine/Source/Developer/TraceAnalysis/Private/Analysis/Engine.h b/Engine/Source/Developer/TraceAnalysis/Private/Analysis/Engine.h
--- a/Engine/Source/Developer/TraceAnalysis/Private/Analysis/Engine.h
+++ b/Engine/Source/Developer/TraceAnalysis/Private/Analysis/Engine.h
@@ -35,6 +35,8 @@ private:
 
 	bool				EstablishTransport(FStreamReader::FData& Data);
 	FDispatch*			AddDispatch(uint16 Uid, uint16 FieldCount=0, uint16 ExtraData=0);
+	template <typename ImplType>
+	void				ForEachRoute(const FDispatch* Dispatch, ImplType&& Impl);
 	void				AddRoute(uint16 AnalyzerIndex, uint16 Id, const ANSICHAR* Logger, const ANSICHAR* Event);
 	void				AddRoute(uint16 AnalyzerIndex, uint16 Id, uint32 Hash);
 	void				OnNewTrace(const FOnEventContext& Context);
