Use default member initialisation for FDispatch objects.
#rnx

diff --git a/Engine/Source/Developer/TraceAnalysis/Private/Analysis/Engine.cpp b/Engine/Source/Developer/TraceAnalysis/Private/Analysis/Engine.cpp
--- a/Engine/Source/Developer/TraceAnalysis/Private/Analysis/Engine.cpp
+++ b/Engine/Source/Developer/TraceAnalysis/Private/Analysis/Engine.cpp
@@ -44,11 +44,11 @@ struct FAnalysisEngine::FDispatch
 		int16		SizeAndType;		// value == byte_size, sign == float < 0 < int
 	};
 
-	uint16			Uid;
-	uint16			FirstRoute;
-	uint16			FieldCount;			// Implicit logger name offset; Fields + FieldCount
-	uint16			EventSize;
-	uint16			EventNameOffset;	// From FDispatch ptr
+	uint16			Uid					= 0;
+	uint16			FirstRoute			= ~uint16(0);
+	uint16			FieldCount			= 0;    // Implicit logger name offset; Fields + FieldCount
+	uint16			EventSize			= 0;
+	uint16			EventNameOffset		= 0;	// From FDispatch ptr
 	uint16			_Unused0;
 	FField			Fields[];
 };
@@ -388,10 +388,9 @@ FAnalysisEngine::FDispatch* FAnalysisEngine::AddDispatch(
 	// Allocate a block of memory to hold the dispatch
 	uint32 Size = sizeof(FDispatch) + (sizeof(FDispatch::FField) * FieldCount) + ExtraData;
 	auto* Dispatch = (FDispatch*)FMemory::Malloc(Size);
+	new (Dispatch) FDispatch();
 	Dispatch->Uid = Uid;
 	Dispatch->FieldCount = FieldCount;
-	Dispatch->EventSize = 0;
-	Dispatch->FirstRoute = ~uint16(0);
 
 	// Add the new dispatch in the dispatch table
 	Dispatches[Uid] = Dispatch;
