Use an explicit logger name offset instead of repurposing fields/field count.
#rnx

diff --git a/Engine/Source/Developer/TraceAnalysis/Private/Analysis/Engine.cpp b/Engine/Source/Developer/TraceAnalysis/Private/Analysis/Engine.cpp
--- a/Engine/Source/Developer/TraceAnalysis/Private/Analysis/Engine.cpp
+++ b/Engine/Source/Developer/TraceAnalysis/Private/Analysis/Engine.cpp
@@ -46,8 +46,9 @@ struct FAnalysisEngine::FDispatch
 
 	uint16			Uid					= 0;
 	uint16			FirstRoute			= ~uint16(0);
-	uint16			FieldCount			= 0;    // Implicit logger name offset; Fields + FieldCount
+	uint16			FieldCount			= 0;
 	uint16			EventSize			= 0;
+	uint16			LoggerNameOffset	= 0;	// From FDispatch ptr
 	uint16			EventNameOffset		= 0;	// From FDispatch ptr
 	uint16			_Unused0;
 	FField			Fields[];
@@ -80,7 +81,7 @@ const ANSICHAR* IAnalyzer::FEventTypeInfo::GetName() const
 const ANSICHAR* IAnalyzer::FEventTypeInfo::GetLoggerName() const
 {
 	const auto* Inner = (const FAnalysisEngine::FDispatch*)this;
-	return (const ANSICHAR*)(Inner->Fields + Inner->FieldCount);
+	return (const ANSICHAR*)Inner + Inner->LoggerNameOffset;
 }
 
 ////////////////////////////////////////////////////////////////////////////////
@@ -493,10 +494,10 @@ void FAnalysisEngine::OnNewEventInternal(const FOnEventContext& Context)
 		WriteCursor[-1] = '\0';
 	};
 
-	UPTRINT EventNameOffset = UPTRINT(Dispatch->Fields + NewEvent.FieldCount);
-	EventNameOffset -= UPTRINT(Dispatch);
-	EventNameOffset += NewEvent.LoggerNameSize + 1;
-	Dispatch->EventNameOffset = uint16(EventNameOffset);
+	UPTRINT LoggerNameOffset = UPTRINT(Dispatch->Fields + NewEvent.FieldCount);
+	LoggerNameOffset -= UPTRINT(Dispatch);
+	Dispatch->LoggerNameOffset = uint16(LoggerNameOffset);
+	Dispatch->EventNameOffset = Dispatch->LoggerNameOffset + NewEvent.LoggerNameSize + 1;
 
 	WriteName(NewEvent.LoggerNameSize);
 	WriteName(NewEvent.EventNameSize);
