Persist the hash calculated for a dispatch
#rnx

diff --git a/Engine/Source/Developer/TraceAnalysis/Private/Analysis/Engine.cpp b/Engine/Source/Developer/TraceAnalysis/Private/Analysis/Engine.cpp
--- a/Engine/Source/Developer/TraceAnalysis/Private/Analysis/Engine.cpp
+++ b/Engine/Source/Developer/TraceAnalysis/Private/Analysis/Engine.cpp
@@ -44,6 +44,7 @@ struct FAnalysisEngine::FDispatch
 		int16		SizeAndType;		// value == byte_size, sign == float < 0 < int
 	};
 
+	uint32			Hash				= 0;
 	uint16			Uid					= 0;
 	uint16			FirstRoute			= ~uint16(0);
 	uint16			FieldCount			= 0;
@@ -462,6 +463,7 @@ void FAnalysisEngine::OnNewEventInternal(const FOnEventContext& Context)
 	FFnv1aHash DispatchHash;
 	NameCursor = DispatchHash.Add(NameCursor, NewEvent.LoggerNameSize);
 	NameCursor = DispatchHash.Add(NameCursor, NewEvent.EventNameSize);
+	Dispatch->Hash = DispatchHash.Get();
 
 	// Fill out the fields
 	for (int i = 0, n = NewEvent.FieldCount; i < n; ++i)
@@ -523,7 +525,7 @@ void FAnalysisEngine::OnNewEventInternal(const FOnEventContext& Context)
 	// Find routes that have subscribed to this event.
 	for (uint16 i = 0, n = Routes.Num(); i < n; ++i)
 	{
-		if (Routes[i].Hash == DispatchHash.Get())
+		if (Routes[i].Hash == Dispatch->Hash)
 		{
 			Dispatch->FirstRoute = i;
 			break;
