Create dispatches via a builder object. This will be useful for maintaining backwards compatibility when there's different trace protocols
#rnx

diff --git a/Engine/Source/Developer/TraceAnalysis/Private/Analysis/Engine.cpp b/Engine/Source/Developer/TraceAnalysis/Private/Analysis/Engine.cpp
--- a/Engine/Source/Developer/TraceAnalysis/Private/Analysis/Engine.cpp
+++ b/Engine/Source/Developer/TraceAnalysis/Private/Analysis/Engine.cpp
@@ -57,6 +57,127 @@ struct FAnalysisEngine::FDispatch
 
 
 
+////////////////////////////////////////////////////////////////////////////////
+class FAnalysisEngine::FDispatchBuilder
+{
+public:
+					FDispatchBuilder();
+	void			SetUid(uint16 Uid);
+	void			SetLoggerName(const ANSICHAR* Name, int32 NameSize);
+	void			SetEventName(const ANSICHAR* Name, int32 NameSize);
+	void			AddField(const ANSICHAR* Name, int32 NameSize, uint16 Offset, uint16 Size, FEventFieldInfo::EType TypeId, uint16 TypeSize);
+	FDispatch*		Finalize();
+
+private:
+	uint32			AppendName(const ANSICHAR* Name, int32 NameSize);
+	TArray<uint8>	Buffer;
+	TArray<uint8>	NameBuffer;
+};
+
+////////////////////////////////////////////////////////////////////////////////
+FAnalysisEngine::FDispatchBuilder::FDispatchBuilder()
+{
+	Buffer.SetNum(sizeof(FDispatch));
+
+	auto* Dispatch = (FDispatch*)(Buffer.GetData());
+	new (Dispatch) FDispatch();
+}
+
+////////////////////////////////////////////////////////////////////////////////
+FAnalysisEngine::FDispatch* FAnalysisEngine::FDispatchBuilder::Finalize()
+{
+	int32 Size = Buffer.Num() + NameBuffer.Num();
+	auto* Dispatch = (FDispatch*)FMemory::Malloc(Size);
+	memcpy(Dispatch, Buffer.GetData(), Buffer.Num());
+	memcpy(Dispatch->Fields + Dispatch->FieldCount, NameBuffer.GetData(), NameBuffer.Num());
+
+	// Sort by hash so we can binary search when looking up.
+	TArrayView<FDispatch::FField> Fields(Dispatch->Fields, Dispatch->FieldCount);
+	Algo::SortBy(Fields, [] (const auto& Field) { return Field.Hash; });
+
+	// Fix up name offsets
+	for (int i = 0, n = Dispatch->FieldCount; i < n; ++i)
+	{
+		auto& Field = Dispatch->Fields[0];
+		Field.NameOffset += sizeof(Field) * (n - i + 1);
+	}
+
+	// Calculate this dispatch's hash.
+	if (Dispatch->LoggerNameOffset || Dispatch->EventNameOffset)
+	{
+		Dispatch->LoggerNameOffset += Buffer.Num();
+		Dispatch->EventNameOffset += Buffer.Num();
+
+		FFnv1aHash Hash;
+		Hash.Add((const ANSICHAR*)Dispatch + Dispatch->LoggerNameOffset);
+		Hash.Add((const ANSICHAR*)Dispatch + Dispatch->EventNameOffset);
+		Dispatch->Hash = Hash.Get();
+	}
+
+	return Dispatch;
+}
+
+////////////////////////////////////////////////////////////////////////////////
+void FAnalysisEngine::FDispatchBuilder::SetUid(uint16 Uid)
+{
+	auto* Dispatch = (FDispatch*)(Buffer.GetData());
+	Dispatch->Uid = Uid;
+}
+
+////////////////////////////////////////////////////////////////////////////////
+void FAnalysisEngine::FDispatchBuilder::SetLoggerName(const ANSICHAR* Name, int32 NameSize)
+{
+	auto* Dispatch = (FDispatch*)(Buffer.GetData());
+	Dispatch->LoggerNameOffset += AppendName(Name, NameSize);
+}
+
+////////////////////////////////////////////////////////////////////////////////
+void FAnalysisEngine::FDispatchBuilder::SetEventName(const ANSICHAR* Name, int32 NameSize)
+{
+	auto* Dispatch = (FDispatch*)(Buffer.GetData());
+	Dispatch->EventNameOffset = AppendName(Name, NameSize);
+}
+
+////////////////////////////////////////////////////////////////////////////////
+void FAnalysisEngine::FDispatchBuilder::AddField(const ANSICHAR* Name, int32 NameSize, uint16 Offset, uint16 Size, FEventFieldInfo::EType TypeId, uint16 TypeSize)
+{
+	int32 Bufoff = Buffer.AddUninitialized(sizeof(FDispatch::FField));
+	auto* Field = (FDispatch::FField*)(Buffer.GetData() + Bufoff);
+	Field->NameOffset = AppendName(Name, NameSize);
+	Field->Offset = Offset;
+	Field->Size = Size;
+	Field->SizeAndType = TypeSize;
+	if (TypeId == FEventFieldInfo::EType::Float)
+	{
+		Field->SizeAndType *= -1;
+	}
+
+	FFnv1aHash Hash;
+	Hash.Add((const uint8*)Name, NameSize);
+	Field->Hash = Hash.Get();
+	
+	auto* Dispatch = (FDispatch*)(Buffer.GetData());
+	Dispatch->FieldCount++;
+	Dispatch->EventSize += Field->Size;
+}
+
+////////////////////////////////////////////////////////////////////////////////
+uint32 FAnalysisEngine::FDispatchBuilder::AppendName(const ANSICHAR* Name, int32 NameSize)
+{
+	if (NameSize < 0)
+	{
+		NameSize = int32(FCStringAnsi::Strlen(Name));
+	}
+
+	int32 Ret = NameBuffer.AddUninitialized(NameSize + 1);
+	uint8* Out = NameBuffer.GetData() + Ret;
+	memcpy(Out, Name, NameSize);
+	Out[NameSize] = '\0';
+	return Ret;
+}
+
+
+
 ////////////////////////////////////////////////////////////////////////////////
 uint32 IAnalyzer::FEventTypeInfo::GetId() const
 {
@@ -229,13 +350,16 @@ FAnalysisEngine::FAnalysisEngine(TArray<IAnalyzer*>&& InAnalyzers)
 	uint16 SelfIndex = Analyzers.Num();
 	Analyzers.Add(this);
 
-	// Manually add event routing for known events, and those we don't quite know
-	// yet but are expecting.
-	FDispatch* NewEventDispatch = AddDispatch(uint16(FNewEventEvent::Uid), 0);
-	NewEventDispatch->FirstRoute = 0;
+	// Manually add event routing for known events.
 	AddRoute(SelfIndex, RouteId_NewEvent, RouteHash_NewEvent);
 	AddRoute(SelfIndex, RouteId_NewTrace, "$Trace", "NewTrace");
 	AddRoute(SelfIndex, RouteId_Timing, "$Trace", "Timing");
+
+	// Add a dispatch for special events
+	FDispatchBuilder Builder;
+	Builder.SetUid(uint16(FNewEventEvent::Uid));
+	Dispatches.Add(Builder.Finalize());
+	Dispatches[0]->FirstRoute = 0;
 }
 
 ////////////////////////////////////////////////////////////////////////////////
@@ -367,38 +491,6 @@ void FAnalysisEngine::OnTiming(const FOnEventContext& Context)
 	SessionContext.CycleFrequency = Context.EventData.GetValue<uint64>("CycleFrequency");
 }
 
-////////////////////////////////////////////////////////////////////////////////
-FAnalysisEngine::FDispatch* FAnalysisEngine::AddDispatch(
-	uint16 Uid,
-	uint16 FieldCount,
-	uint16 ExtraData)
-{
-	// Make sure there's enough space in the dispatch table, failing gently if
-	// there appears to be an existing entry.
-	if (Uid < Dispatches.Num())
-	{
-		if (Dispatches[Uid] != nullptr)
-		{
-			return nullptr;
-		}
-	}
-	else
-	{
-		Dispatches.SetNum(Uid + 1);
-	}
-
-	// Allocate a block of memory to hold the dispatch
-	uint32 Size = sizeof(FDispatch) + (sizeof(FDispatch::FField) * FieldCount) + ExtraData;
-	auto* Dispatch = (FDispatch*)FMemory::Malloc(Size);
-	new (Dispatch) FDispatch();
-	Dispatch->Uid = Uid;
-	Dispatch->FieldCount = FieldCount;
-
-	// Add the new dispatch in the dispatch table
-	Dispatches[Uid] = Dispatch;
-	return Dispatch;
-}
-
 ////////////////////////////////////////////////////////////////////////////////
 template <typename ImplType>
 void FAnalysisEngine::ForEachRoute(const FDispatch* Dispatch, ImplType&& Impl)
@@ -437,90 +529,54 @@ void FAnalysisEngine::OnNewEventInternal(const FOnEventContext& Context)
 	const FEventDataInfo& EventData = (const FEventDataInfo&)(Context.EventData);
 	const auto& NewEvent = *(FNewEventEvent*)(EventData.Ptr);
 
-	// Create a new dispatch with enough space to store new event's various names
-	uint16 NameDataSize = NewEvent.LoggerNameSize + NewEvent.EventNameSize;
-	NameDataSize += NewEvent.FieldCount + 2; // null terminators
-	for (int i = 0, n = NewEvent.FieldCount; i < n; ++i)
-	{
-		NameDataSize += NewEvent.Fields[i].NameSize + 1;
-	}
-
-	FDispatch* Dispatch = AddDispatch(NewEvent.EventUid, NewEvent.FieldCount, NameDataSize);
-	if (Dispatch == nullptr)
-	{
-		return;
-	}
+	FDispatchBuilder Builder;
 
-	if (NewEvent.FieldCount)
-	{
-		auto& LastField = NewEvent.Fields[NewEvent.FieldCount - 1];
-		Dispatch->EventSize = LastField.Offset + LastField.Size;
-	}
+	const auto* NameCursor = (const ANSICHAR*)(NewEvent.Fields + NewEvent.FieldCount);
 
-	const uint8* NameCursor = (const uint8*)(NewEvent.Fields + NewEvent.FieldCount);
+	Builder.SetLoggerName(NameCursor, NewEvent.LoggerNameSize);
+	NameCursor += NewEvent.LoggerNameSize;
 
-	// Calculate this dispatch's hash.
-	FFnv1aHash DispatchHash;
-	NameCursor = DispatchHash.Add(NameCursor, NewEvent.LoggerNameSize);
-	NameCursor = DispatchHash.Add(NameCursor, NewEvent.EventNameSize);
-	Dispatch->Hash = DispatchHash.Get();
+	Builder.SetEventName(NameCursor, NewEvent.EventNameSize);
+	NameCursor += NewEvent.EventNameSize;
+	Builder.SetUid(NewEvent.EventUid);
 
 	// Fill out the fields
 	for (int i = 0, n = NewEvent.FieldCount; i < n; ++i)
 	{
-		const auto& In = NewEvent.Fields[i];
-		auto& Out = Dispatch->Fields[i];
-
-		FFnv1aHash FieldHash;
-		NameCursor = FieldHash.Add(NameCursor, In.NameSize);
+		const auto& Field = NewEvent.Fields[i];
 
-		Out.Hash = FieldHash.Get();
-		Out.Offset = In.Offset;
-		Out.Size = In.Size;
+		uint16 TypeSize = 1 << (Field.TypeInfo & _Field_Pow2SizeMask);
 
-		Out.SizeAndType = 1 << (In.TypeInfo & _Field_Pow2SizeMask);
-		if ((In.TypeInfo & _Field_CategoryMask) == _Field_Float)
+		auto TypeId = FEventFieldInfo::EType::Integer;
+		if ((Field.TypeInfo & _Field_CategoryMask) == _Field_Float)
 		{
-			Out.SizeAndType = -Out.SizeAndType;
+			TypeId = FEventFieldInfo::EType::Float;
 		}
-	}
-
-	// Write out names with null terminators.
-	NameCursor = (const uint8*)(NewEvent.Fields + NewEvent.FieldCount);
-	uint8* WriteCursor = (uint8*)(Dispatch->Fields + Dispatch->FieldCount);
-	auto WriteName = [&] (uint32 Size)
-	{
-		memcpy(WriteCursor, NameCursor, Size);
-		NameCursor += Size;
-		WriteCursor += Size + 1;
-		WriteCursor[-1] = '\0';
-	};
-
-	UPTRINT LoggerNameOffset = UPTRINT(Dispatch->Fields + NewEvent.FieldCount);
-	LoggerNameOffset -= UPTRINT(Dispatch);
-	Dispatch->LoggerNameOffset = uint16(LoggerNameOffset);
-	Dispatch->EventNameOffset = Dispatch->LoggerNameOffset + NewEvent.LoggerNameSize + 1;
 
-	WriteName(NewEvent.LoggerNameSize);
-	WriteName(NewEvent.EventNameSize);
-	for (int i = 0, n = NewEvent.FieldCount; i < n; ++i)
-	{
-		auto& Out = Dispatch->Fields[i];
-		Out.NameOffset = uint16(UPTRINT(WriteCursor) - UPTRINT(Dispatch));
+		Builder.AddField(NameCursor, Field.NameSize, Field.Offset, Field.Size, TypeId, TypeSize);
 
-		WriteName(NewEvent.Fields[i].NameSize);
+		NameCursor += Field.NameSize;
 	}
 
-	// Sort by hash so we can binary search when looking up.
-	TArrayView<FDispatch::FField> Fields(Dispatch->Fields, Dispatch->FieldCount);
-	Algo::SortBy(Fields, [] (const auto& Field) { return Field.Hash; });
+	// Get the dispatch and add it into the dispatch table. Fail gently if there
+	// is the dispatch table unexpetedly already has an entry.
+	FDispatch* Dispatch = Builder.Finalize();
 
-	// Fix up field name offsets
-	for (int i = 0, n = NewEvent.FieldCount; i < n; ++i)
-	{
-		auto& Out = Dispatch->Fields[i];
-		Out.NameOffset = uint16(UPTRINT(Dispatch) + Out.NameOffset - UPTRINT(&Out));
-	}
+	uint16 Uid = Dispatch->Uid;
+	if (Uid < Dispatches.Num())
+ 	{
+		if (Dispatches[Uid] != nullptr)
+ 		{
+			FMemory::Free(Dispatch);
+			return;
+ 		}
+ 	}
+	else
+ 	{
+		Dispatches.SetNum(Uid + 1);
+ 	}
+
+	Dispatches[Uid] = Dispatch;
 
 	// Find routes that have subscribed to this event.
 	for (uint16 i = 0, n = Routes.Num(); i < n; ++i)
diff --git a/Engine/Source/Developer/TraceAnalysis/Private/Analysis/Engine.h b/Engine/Source/Developer/TraceAnalysis/Private/Analysis/Engine.h
--- a/Engine/Source/Developer/TraceAnalysis/Private/Analysis/Engine.h
+++ b/Engine/Source/Developer/TraceAnalysis/Private/Analysis/Engine.h
@@ -31,10 +31,10 @@ private:
 		uint16			_Unused0;
 	};
 
+	class				FDispatchBuilder;
 	virtual bool		OnEvent(uint16 RouteId, const FOnEventContext& Context) override;
 
 	bool				EstablishTransport(FStreamReader::FData& Data);
-	FDispatch*			AddDispatch(uint16 Uid, uint16 FieldCount=0, uint16 ExtraData=0);
 	template <typename ImplType>
 	void				ForEachRoute(const FDispatch* Dispatch, ImplType&& Impl);
 	void				AddRoute(uint16 AnalyzerIndex, uint16 Id, const ANSICHAR* Logger, const ANSICHAR* Event);
