Wrapped the adding of dispatches in a method so that a new dispatch can be added from anywhere in the analysis engine and these new dispatches will be connected to their subscribed routes.

#rnx

diff --git a/Engine/Source/Developer/TraceAnalysis/Private/Analysis/Engine.cpp b/Engine/Source/Developer/TraceAnalysis/Private/Analysis/Engine.cpp
--- a/Engine/Source/Developer/TraceAnalysis/Private/Analysis/Engine.cpp
+++ b/Engine/Source/Developer/TraceAnalysis/Private/Analysis/Engine.cpp
@@ -358,8 +358,7 @@ FAnalysisEngine::FAnalysisEngine(TArray<IAnalyzer*>&& InAnalyzers)
 	// Add a dispatch for special events
 	FDispatchBuilder Builder;
 	Builder.SetUid(uint16(FNewEventEvent::Uid));
-	Dispatches.Add(Builder.Finalize());
-	Dispatches[0]->FirstRoute = 0;
+	AddDispatch(Builder.Finalize());
 }
 
 ////////////////////////////////////////////////////////////////////////////////
@@ -561,14 +560,31 @@ void FAnalysisEngine::OnNewEventInternal(const FOnEventContext& Context)
 	// Get the dispatch and add it into the dispatch table. Fail gently if there
 	// is the dispatch table unexpetedly already has an entry.
 	FDispatch* Dispatch = Builder.Finalize();
+	if (!AddDispatch(Dispatch))
+	{
+		return;
+	}
+
+	// Inform routes that a new event has been declared.
+	ForEachRoute(Dispatch, [&] (IAnalyzer* Analyzer, uint16 RouteId)
+	{
+		if (!Analyzer->OnNewEvent(RouteId, *(FEventTypeInfo*)Dispatch))
+		{
+			RetireAnalyzer(Analyzer);
+		}
+	});
+}
 
+////////////////////////////////////////////////////////////////////////////////
+bool FAnalysisEngine::AddDispatch(FDispatch* Dispatch)
+{
 	uint16 Uid = Dispatch->Uid;
 	if (Uid < Dispatches.Num())
  	{
 		if (Dispatches[Uid] != nullptr)
  		{
 			FMemory::Free(Dispatch);
-			return;
+			return false;
  		}
  	}
 	else
@@ -588,14 +604,7 @@ void FAnalysisEngine::OnNewEventInternal(const FOnEventContext& Context)
 		}
 	}
 
-	// Inform routes that a new event has been declared.
-	ForEachRoute(Dispatch, [&] (IAnalyzer* Analyzer, uint16 RouteId)
-	{
-		if (!Analyzer->OnNewEvent(RouteId, *(FEventTypeInfo*)Dispatch))
-		{
-			RetireAnalyzer(Analyzer);
-		}
-	});
+	return true;
 }
 
 ////////////////////////////////////////////////////////////////////////////////
diff --git a/Engine/Source/Developer/TraceAnalysis/Private/Analysis/Engine.h b/Engine/Source/Developer/TraceAnalysis/Private/Analysis/Engine.h
--- a/Engine/Source/Developer/TraceAnalysis/Private/Analysis/Engine.h
+++ b/Engine/Source/Developer/TraceAnalysis/Private/Analysis/Engine.h
@@ -35,6 +35,7 @@ private:
 	virtual bool		OnEvent(uint16 RouteId, const FOnEventContext& Context) override;
 
 	bool				EstablishTransport(FStreamReader::FData& Data);
+	bool				AddDispatch(FDispatch* Dispatch);
 	template <typename ImplType>
 	void				ForEachRoute(const FDispatch* Dispatch, ImplType&& Impl);
 	void				AddRoute(uint16 AnalyzerIndex, uint16 Id, const ANSICHAR* Logger, const ANSICHAR* Event);
