Simplified hashing helper as it no longer needs to return the "next" pointer

#rnx

diff --git a/Engine/Source/Developer/TraceAnalysis/Private/Analysis/Engine.cpp b/Engine/Source/Developer/TraceAnalysis/Private/Analysis/Engine.cpp
--- a/Engine/Source/Developer/TraceAnalysis/Private/Analysis/Engine.cpp
+++ b/Engine/Source/Developer/TraceAnalysis/Private/Analysis/Engine.cpp
@@ -17,15 +17,15 @@ namespace Trace
 class FFnv1aHash
 {
 public:
-					FFnv1aHash() = default;
-					FFnv1aHash(uint32 PrevResult)		{ Result = PrevResult; }
-	void			Add(const ANSICHAR* String)			{ for (; *String; ++String) { Add(*String); } }
-	const uint8*	Add(const uint8* Data, uint32 Size)	{ for (uint32 i = 0; i < Size; ++Data, ++i) { Add(*Data); } return Data; }
-	void			Add(uint8 Value)					{ Result ^= Value; Result *= 0x01000193; }
-	uint32			Get() const							{ return Result; }
+			FFnv1aHash() = default;
+			FFnv1aHash(uint32 PrevResult)		{ Result = PrevResult; }
+	void	Add(const ANSICHAR* String)			{ for (; *String; ++String) { Add(*String); } }
+	void	Add(const uint8* Data, uint32 Size)	{ for (uint32 i = 0; i < Size; ++Data, ++i) { Add(*Data); } }
+	void	Add(uint8 Value)					{ Result ^= Value; Result *= 0x01000193; }
+	uint32	Get() const							{ return Result; }
 
 private:
-	uint32			Result = 0x811c9dc5;
+	uint32	Result = 0x811c9dc5;
 	// uint32: bias=0x811c9dc5			prime=0x01000193
 	// uint64: bias=0xcbf29ce484222325	prime=0x00000100000001b3;
 };
