Prefer the term "packet" over "payload"

#rnx

diff --git a/Engine/Source/Developer/TraceAnalysis/Private/Analysis/Transport/TlsTransport.cpp b/Engine/Source/Developer/TraceAnalysis/Private/Analysis/Transport/TlsTransport.cpp
--- a/Engine/Source/Developer/TraceAnalysis/Private/Analysis/Transport/TlsTransport.cpp
+++ b/Engine/Source/Developer/TraceAnalysis/Private/Analysis/Transport/TlsTransport.cpp
@@ -8,11 +8,11 @@ namespace Trace
 ////////////////////////////////////////////////////////////////////////////////
 FTlsTransport::~FTlsTransport()
 {
-	for (FPayloadNode* Root : {ActiveList, PendingList, FreeList})
+	for (FPacketNode* Root : {ActiveList, PendingList, FreeList})
 	{
-		for (FPayloadNode* Node = Root; Node != nullptr;)
+		for (FPacketNode* Node = Root; Node != nullptr;)
 		{
-			FPayloadNode* Next = Node->Next;
+			FPacketNode* Next = Node->Next;
 			delete[] Node;
 			Node = Next;
 		}
@@ -39,7 +39,7 @@ const uint8* FTlsTransport::GetPointerImpl(uint32 BlockSize)
 	uint32 NextCursor = ActiveList->Cursor + BlockSize;
 	if (NextCursor > ActiveList->Size)
 	{
-		FPayloadNode* Node = ActiveList;
+		FPacketNode* Node = ActiveList;
 		ActiveList = ActiveList->Next;
 		Node->Next = FreeList;
 		FreeList = Node;
@@ -52,7 +52,7 @@ const uint8* FTlsTransport::GetPointerImpl(uint32 BlockSize)
 ////////////////////////////////////////////////////////////////////////////////
 bool FTlsTransport::GetNextBatch()
 {
-	struct FPayload
+	struct FPacket
 	{
 		int16	Serial;
 		uint16	Size;
@@ -67,29 +67,29 @@ bool FTlsTransport::GetNextBatch()
 
 	while (true)
 	{
-		const auto* Payload = (const FPayload*)FTransport::GetPointerImpl(sizeof(FPayload));
-		if (Payload == nullptr)
+		const auto* Packet = (const FPacket*)FTransport::GetPointerImpl(sizeof(FPacket));
+		if (Packet == nullptr)
 		{
 			return false;
 		}
 
 		// If this new payload is part of the next event batch then we've finished
 		// building the current batch. The current batch can be activated.
-		if (LastSerial >= Payload->Serial)
+		if (LastSerial >= Packet->Serial)
 		{
 			ActiveList = PendingList;
 			PendingList = nullptr;
 			break;
 		}
 
-		if (FTransport::GetPointerImpl(Payload->Size) == nullptr)
+		if (FTransport::GetPointerImpl(Packet->Size) == nullptr)
 		{
 			return false;
 		}
 
-		FTransport::Advance(Payload->Size);
+		FTransport::Advance(Packet->Size);
 
-		FPayloadNode* Node;
+		FPacketNode* Node;
 		if (FreeList != nullptr)
 		{
 			Node = FreeList;
@@ -97,17 +97,17 @@ bool FTlsTransport::GetNextBatch()
 		}
 		else
 		{
-			Node = (FPayloadNode*)GMalloc->Malloc(sizeof(FPayloadNode) + MaxPayloadSize);
+			Node = (FPacketNode*)GMalloc->Malloc(sizeof(FPacketNode) + MaxPacketSize);
 		}
 
-		Node->Serial = Payload->Serial;
+		Node->Serial = Packet->Serial;
 		Node->Cursor = 0;
-		Node->Size = uint16(Payload->Size - sizeof(*Payload));
+		Node->Size = uint16(Packet->Size - sizeof(*Packet));
 		Node->Next = PendingList;
-		FMemory::Memcpy(Node->Data, Payload->Data, Node->Size);
+		FMemory::Memcpy(Node->Data, Packet->Data, Node->Size);
 
 		PendingList = Node;
-		LastSerial = Payload->Serial;
+		LastSerial = Packet->Serial;
 	}
 
 	return true;
diff --git a/Engine/Source/Developer/TraceAnalysis/Private/Analysis/Transport/TlsTransport.h b/Engine/Source/Developer/TraceAnalysis/Private/Analysis/Transport/TlsTransport.h
--- a/Engine/Source/Developer/TraceAnalysis/Private/Analysis/Transport/TlsTransport.h
+++ b/Engine/Source/Developer/TraceAnalysis/Private/Analysis/Transport/TlsTransport.h
@@ -17,9 +17,9 @@ public:
 	virtual const uint8*	GetPointerImpl(uint32 BlockSize) override;
 
 private:
-	struct FPayloadNode
+	struct FPacketNode
 	{
-		FPayloadNode*		Next;
+		FPacketNode*		Next;
 		uint32				Cursor;
 		uint16				Serial;
 		uint16				Size;
@@ -27,10 +27,10 @@ private:
 	};
 
 	bool					GetNextBatch();
-	static const uint32		MaxPayloadSize = 8192;
-	FPayloadNode*			ActiveList = nullptr;
-	FPayloadNode*			PendingList = nullptr;
-	FPayloadNode*			FreeList = nullptr;
+	static const uint32		MaxPacketSize = 8192;
+	FPacketNode*			ActiveList = nullptr;
+	FPacketNode*			PendingList = nullptr;
+	FPacketNode*			FreeList = nullptr;
 };
 
 } // namespace Trace
diff --git a/Engine/Source/Runtime/TraceLog/Private/Trace/Writer.cpp b/Engine/Source/Runtime/TraceLog/Private/Trace/Writer.cpp
--- a/Engine/Source/Runtime/TraceLog/Private/Trace/Writer.cpp
+++ b/Engine/Source/Runtime/TraceLog/Private/Trace/Writer.cpp
@@ -352,7 +352,7 @@ static void Writer_ConsumeEvents()
 
 	struct FCollector
 	{
-		struct FPayload
+		struct FPacket
 		{
 			struct FHeader
 			{
@@ -392,23 +392,23 @@ static void Writer_ConsumeEvents()
 
 		void Flush()
 		{
-			if (Cursor == sizeof(FPayload::Data))
+			if (Cursor == sizeof(FPacket::Data))
 			{
 				return;
 			}
 
 			// There will always be space remaining for the header because we've
-			// arranged for that by including the header in FPayload. We'll shift
+			// arranged for that by including the header in FPacket. We'll shift
 			// it forward so it butts up against the event data (at the expense of
 			// some occasional unaligned stores).
-			Cursor -= sizeof(FPayload::FHeader);
+			Cursor -= sizeof(FPacket::FHeader);
 
-			auto* Out = (FPayload::FHeader*)(Payload.Data + Cursor);
+			auto* Out = (FPacket::FHeader*)(Packet.Data + Cursor);
 			Out->Serial = Serial;
-			Out->Size = sizeof(FPayload::Data) - Cursor;
+			Out->Size = sizeof(FPacket::Data) - Cursor;
 			Send(Out, Out->Size);
 
-			Cursor = sizeof(FPayload::Data);
+			Cursor = sizeof(FPacket::Data);
 			Serial++;
 		}
 
@@ -420,12 +420,12 @@ static void Writer_ConsumeEvents()
 			}
 
 			Cursor -= Size;
-			memcpy(Payload.Data + Cursor, Data, Size);
+			memcpy(Packet.Data + Cursor, Data, Size);
 		}
 
-		int16		Cursor = sizeof(FPayload::Data);
+		int16		Cursor = sizeof(FPacket::Data);
 		uint16		Serial = 0;
-		FPayload	Payload;
+		FPacket		Packet;
 	};
 
 	FCollector Collector;
