Renamed TLS transport to payload transport as it really has nothing to do with TLS.

#rnx

diff --git a/Engine/Source/Developer/TraceAnalysis/Private/Analysis/Engine.cpp b/Engine/Source/Developer/TraceAnalysis/Private/Analysis/Engine.cpp
--- a/Engine/Source/Developer/TraceAnalysis/Private/Analysis/Engine.cpp
+++ b/Engine/Source/Developer/TraceAnalysis/Private/Analysis/Engine.cpp
@@ -8,7 +8,7 @@
 #include "Trace/Analyzer.h"
 #include "Trace/Detail/EventDef.h"
 #include "Transport/Transport.h"
-#include "Transport/TlsTransport.h"
+#include "Transport/PayloadTransport.h"
 
 namespace Trace
 {
@@ -642,7 +642,7 @@ bool FAnalysisEngine::EstablishTransport(FStreamReader::FData& Data)
 	switch (Header->Format)
 	{
 	case 1:		Transport = new FTransport(); break;
-	case 2:		Transport = new FTlsTransport(); break;
+	case 2:		Transport = new FPayloadTransport(); break;
 	default:	return false;
 	//case 'E':	/* See the magic above */ break;
 	//case 'T':	/* See the magic above */ break;
diff --git a/Engine/Source/Developer/TraceAnalysis/Private/Analysis/Transport/PayloadTransport.cpp b/Engine/Source/Developer/TraceAnalysis/Private/Analysis/Transport/PayloadTransport.cpp
new file mode 100644
--- /dev/null
+++ b/Engine/Source/Developer/TraceAnalysis/Private/Analysis/Transport/PayloadTransport.cpp
@@ -0,0 +1,116 @@
+// Copyright 1998-2019 Epic Games, Inc. All Rights Reserved.
+
+#include "PayloadTransport.h"
+
+namespace Trace
+{
+
+////////////////////////////////////////////////////////////////////////////////
+FPayloadTransport::~FPayloadTransport()
+{
+	for (FPacketNode* Root : {ActiveList, PendingList, FreeList})
+	{
+		for (FPacketNode* Node = Root; Node != nullptr;)
+		{
+			FPacketNode* Next = Node->Next;
+			delete[] Node;
+			Node = Next;
+		}
+	}
+}
+
+////////////////////////////////////////////////////////////////////////////////
+void FPayloadTransport::Advance(uint32 BlockSize)
+{
+	if (ActiveList != nullptr)
+	{
+		ActiveList->Cursor += BlockSize;
+	}
+}
+
+////////////////////////////////////////////////////////////////////////////////
+const uint8* FPayloadTransport::GetPointerImpl(uint32 BlockSize)
+{
+	if (ActiveList == nullptr && !GetNextBatch())
+	{
+		return nullptr;
+	}
+
+	uint32 NextCursor = ActiveList->Cursor + BlockSize;
+	if (NextCursor > ActiveList->Size)
+	{
+		FPacketNode* Node = ActiveList;
+		ActiveList = ActiveList->Next;
+		Node->Next = FreeList;
+		FreeList = Node;
+		return GetPointerImpl(BlockSize);
+	}
+
+	return ActiveList->Data + ActiveList->Cursor;
+}
+
+////////////////////////////////////////////////////////////////////////////////
+bool FPayloadTransport::GetNextBatch()
+{
+	struct FPacket
+	{
+		int16	Serial;
+		uint16	Size;
+		uint8	Data[];
+	};
+
+	int16 LastSerial = -1;
+	if (PendingList != nullptr)
+	{
+		LastSerial = PendingList->Serial;
+	}
+
+	while (true)
+	{
+		const auto* Packet = (const FPacket*)FTransport::GetPointerImpl(sizeof(FPacket));
+		if (Packet == nullptr)
+		{
+			return false;
+		}
+
+		// If this new payload is part of the next event batch then we've finished
+		// building the current batch. The current batch can be activated.
+		if (LastSerial >= Packet->Serial)
+		{
+			ActiveList = PendingList;
+			PendingList = nullptr;
+			break;
+		}
+
+		if (FTransport::GetPointerImpl(Packet->Size) == nullptr)
+		{
+			return false;
+		}
+
+		FTransport::Advance(Packet->Size);
+
+		FPacketNode* Node;
+		if (FreeList != nullptr)
+		{
+			Node = FreeList;
+			FreeList = Node->Next;
+		}
+		else
+		{
+			Node = (FPacketNode*)GMalloc->Malloc(sizeof(FPacketNode) + MaxPacketSize);
+		}
+
+		Node->Serial = Packet->Serial;
+		Node->Cursor = 0;
+		Node->Size = uint16(Packet->Size - sizeof(*Packet));
+		Node->Next = PendingList;
+		FMemory::Memcpy(Node->Data, Packet->Data, Node->Size);
+
+		PendingList = Node;
+		LastSerial = Packet->Serial;
+	}
+
+	return true;
+}
+
+} // namespace Trace
diff --git a/Engine/Source/Developer/TraceAnalysis/Private/Analysis/Transport/PayloadTransport.h b/Engine/Source/Developer/TraceAnalysis/Private/Analysis/Transport/PayloadTransport.h
new file mode 100644
--- /dev/null
+++ b/Engine/Source/Developer/TraceAnalysis/Private/Analysis/Transport/PayloadTransport.h
@@ -0,0 +1,36 @@
+// Copyright 1998-2019 Epic Games, Inc. All Rights Reserved.
+
+#pragma once
+
+#include "Transport.h"
+
+namespace Trace
+{
+
+////////////////////////////////////////////////////////////////////////////////
+class FPayloadTransport
+	: public FTransport
+{
+public:
+	virtual					~FPayloadTransport();
+	virtual void			Advance(uint32 BlockSize) override;
+	virtual const uint8*	GetPointerImpl(uint32 BlockSize) override;
+
+private:
+	struct FPacketNode
+	{
+		FPacketNode*		Next;
+		uint32				Cursor;
+		uint16				Serial;
+		uint16				Size;
+		uint8				Data[];
+	};
+
+	bool					GetNextBatch();
+	static const uint32		MaxPacketSize = 8192;
+	FPacketNode*			ActiveList = nullptr;
+	FPacketNode*			PendingList = nullptr;
+	FPacketNode*			FreeList = nullptr;
+};
+
+} // namespace Trace
diff --git a/Engine/Source/Developer/TraceAnalysis/Private/Analysis/Transport/TlsTransport.cpp b/Engine/Source/Developer/TraceAnalysis/Private/Analysis/Transport/TlsTransport.cpp
deleted file mode 100644
diff --git a/Engine/Source/Developer/TraceAnalysis/Private/Analysis/Transport/TlsTransport.h b/Engine/Source/Developer/TraceAnalysis/Private/Analysis/Transport/TlsTransport.h
deleted file mode 100644
