Hid FPacketPayload::FPacketNode

#rnx

diff --git a/Engine/Source/Developer/TraceAnalysis/Private/Analysis/Transport/PacketTransport.cpp b/Engine/Source/Developer/TraceAnalysis/Private/Analysis/Transport/PacketTransport.cpp
--- a/Engine/Source/Developer/TraceAnalysis/Private/Analysis/Transport/PacketTransport.cpp
+++ b/Engine/Source/Developer/TraceAnalysis/Private/Analysis/Transport/PacketTransport.cpp
@@ -5,6 +5,16 @@
 namespace Trace
 {
 
+////////////////////////////////////////////////////////////////////////////////
+struct FPayloadTransport::FPacketNode
+{
+	FPacketNode*		Next;
+	uint32				Cursor;
+	uint16				Serial;
+	uint16				Size;
+	uint8				Data[];
+};
+
 ////////////////////////////////////////////////////////////////////////////////
 FPayloadTransport::~FPayloadTransport()
 {
diff --git a/Engine/Source/Developer/TraceAnalysis/Private/Analysis/Transport/PacketTransport.h b/Engine/Source/Developer/TraceAnalysis/Private/Analysis/Transport/PacketTransport.h
--- a/Engine/Source/Developer/TraceAnalysis/Private/Analysis/Transport/PacketTransport.h
+++ b/Engine/Source/Developer/TraceAnalysis/Private/Analysis/Transport/PacketTransport.h
@@ -17,15 +17,7 @@ public:
 	virtual const uint8*	GetPointerImpl(uint32 BlockSize) override;
 
 private:
-	struct FPacketNode
-	{
-		FPacketNode*		Next;
-		uint32				Cursor;
-		uint16				Serial;
-		uint16				Size;
-		uint8				Data[];
-	};
-
+	struct					FPacketNode;
 	bool					GetNextBatch();
 	static const uint32		MaxPacketSize = 8192;
 	FPacketNode*			ActiveList = nullptr;
