Split out node allocation from the middle of parsing trace transport packets

#rnx

diff --git a/Engine/Source/Developer/TraceAnalysis/Private/Analysis/Transport/PacketTransport.cpp b/Engine/Source/Developer/TraceAnalysis/Private/Analysis/Transport/PacketTransport.cpp
--- a/Engine/Source/Developer/TraceAnalysis/Private/Analysis/Transport/PacketTransport.cpp
+++ b/Engine/Source/Developer/TraceAnalysis/Private/Analysis/Transport/PacketTransport.cpp
@@ -59,6 +59,24 @@ const uint8* FPayloadTransport::GetPointerImpl(uint32 BlockSize)
 	return ActiveList->Data + ActiveList->Cursor;
 }
 
+////////////////////////////////////////////////////////////////////////////////
+FPayloadTransport::FPacketNode* FPayloadTransport::AllocateNode()
+{
+	FPacketNode* Node;
+	if (FreeList != nullptr)
+	{
+		Node = FreeList;
+		FreeList = Node->Next;
+	}
+	else
+	{
+		Node = (FPacketNode*)GMalloc->Malloc(sizeof(FPacketNode) + MaxPacketSize);
+	}
+
+	Node->Cursor = 0;
+	return Node;
+}
+
 ////////////////////////////////////////////////////////////////////////////////
 bool FPayloadTransport::GetNextBatch()
 {
@@ -99,17 +117,7 @@ bool FPayloadTransport::GetNextBatch()
 
 		FTransport::Advance(Packet->Size);
 
-		FPacketNode* Node;
-		if (FreeList != nullptr)
-		{
-			Node = FreeList;
-			FreeList = Node->Next;
-		}
-		else
-		{
-			Node = (FPacketNode*)GMalloc->Malloc(sizeof(FPacketNode) + MaxPacketSize);
-		}
-
+		FPacketNode* Node = AllocateNode();
 		Node->Serial = Packet->Serial;
 		Node->Cursor = 0;
 		Node->Size = uint16(Packet->Size - sizeof(*Packet));
diff --git a/Engine/Source/Developer/TraceAnalysis/Private/Analysis/Transport/PacketTransport.h b/Engine/Source/Developer/TraceAnalysis/Private/Analysis/Transport/PacketTransport.h
--- a/Engine/Source/Developer/TraceAnalysis/Private/Analysis/Transport/PacketTransport.h
+++ b/Engine/Source/Developer/TraceAnalysis/Private/Analysis/Transport/PacketTransport.h
@@ -19,6 +19,7 @@ public:
 private:
 	struct					FPacketNode;
 	bool					GetNextBatch();
+	FPacketNode*			AllocateNode();
 	static const uint32		MaxPacketSize = 8192;
 	FPacketNode*			ActiveList = nullptr;
 	FPacketNode*			PendingList = nullptr;
