Implemented trace stream compression initially disabled pending completion of paperwork.

#rnx

diff --git a/Engine/Source/Developer/TraceAnalysis/Private/Analysis/Transport/PacketTransport.cpp b/Engine/Source/Developer/TraceAnalysis/Private/Analysis/Transport/PacketTransport.cpp
--- a/Engine/Source/Developer/TraceAnalysis/Private/Analysis/Transport/PacketTransport.cpp
+++ b/Engine/Source/Developer/TraceAnalysis/Private/Analysis/Transport/PacketTransport.cpp
@@ -5,6 +5,16 @@
 namespace Trace
 {
 
+////////////////////////////////////////////////////////////////////////////////
+namespace Private
+{
+
+TRACELOG_API int32 Decode(const void*, int32, void*, int32);
+
+} // namespace Private
+
+
+
 ////////////////////////////////////////////////////////////////////////////////
 struct FPayloadTransport::FPacketNode
 {
@@ -80,13 +90,6 @@ FPayloadTransport::FPacketNode* FPayloadTransport::AllocateNode()
 ////////////////////////////////////////////////////////////////////////////////
 bool FPayloadTransport::GetNextBatch()
 {
-	struct FPacket
-	{
-		int16	Serial;
-		uint16	Size;
-		uint8	Data[];
-	};
-
 	int16 LastSerial = -1;
 	if (PendingList != nullptr)
 	{
@@ -95,37 +98,74 @@ bool FPayloadTransport::GetNextBatch()
 
 	while (true)
 	{
-		const auto* Packet = (const FPacket*)FTransport::GetPointerImpl(sizeof(FPacket));
-		if (Packet == nullptr)
+		struct FPacketBase
+		{
+			uint16	Serial;
+			uint16	PacketSize;
+		};
+
+		const auto* PacketBase = (const FPacketBase*)FTransport::GetPointerImpl(sizeof(FPacketBase));
+		if (PacketBase == nullptr)
 		{
 			return false;
 		}
 
 		// If this new payload is part of the next event batch then we've finished
 		// building the current batch. The current batch can be activated.
-		if (LastSerial >= Packet->Serial)
+		int16 PacketSerial = (PacketBase->Serial & 0x7fff);
+		if (LastSerial >= PacketSerial)
 		{
 			ActiveList = PendingList;
 			PendingList = nullptr;
 			break;
 		}
 
-		if (FTransport::GetPointerImpl(Packet->Size) == nullptr)
+		if (FTransport::GetPointerImpl(PacketBase->PacketSize) == nullptr)
 		{
 			return false;
 		}
 
-		FTransport::Advance(Packet->Size);
+		FTransport::Advance(PacketBase->PacketSize);
+
+		LastSerial = PacketSerial;
 
 		FPacketNode* Node = AllocateNode();
-		Node->Serial = Packet->Serial;
-		Node->Cursor = 0;
-		Node->Size = uint16(Packet->Size - sizeof(*Packet));
+		Node->Serial = PacketSerial;
 		Node->Next = PendingList;
-		FMemory::Memcpy(Node->Data, Packet->Data, Node->Size);
-
 		PendingList = Node;
-		LastSerial = Packet->Serial;
+
+#if 0
+		bool bEncoded = (PacketBase->Serial != PacketSerial);
+		if (bEncoded)
+		{
+			struct FPacketEncoded
+				: public FPacketBase
+			{
+				uint16	DecodedSize;
+				uint8	Data[];
+			};
+			auto* PacketEncoded = (FPacketEncoded*)PacketBase;
+
+			Node->Size = Private::Decode(
+				PacketEncoded->Data,
+				int32(PacketEncoded->PacketSize - sizeof(FPacketEncoded)),
+				Node->Data,
+				PacketEncoded->DecodedSize
+			);
+		}
+		else
+#endif // 0
+		{
+			struct FPacketRaw
+				: public FPacketBase
+			{
+				uint8 Data[];
+			};
+			auto* PacketRaw = (FPacketRaw*)PacketBase;
+
+			Node->Size = uint16(PacketBase->PacketSize - sizeof(FPacketBase));
+			FMemory::Memcpy(Node->Data, PacketRaw->Data, Node->Size);
+		}
 	}
 
 	return true;
diff --git a/Engine/Source/Runtime/TraceLog/Private/Trace/Writer.cpp b/Engine/Source/Runtime/TraceLog/Private/Trace/Writer.cpp
--- a/Engine/Source/Runtime/TraceLog/Private/Trace/Writer.cpp
+++ b/Engine/Source/Runtime/TraceLog/Private/Trace/Writer.cpp
@@ -23,6 +23,11 @@
 namespace Trace {
 namespace Private {
 
+////////////////////////////////////////////////////////////////////////////////
+int32 Encode(const void*, int32, void*, int32);
+
+
+
 ////////////////////////////////////////////////////////////////////////////////
 inline void Writer_Yield()
 {
@@ -350,19 +355,9 @@ static void Writer_ConsumeEvents()
 	FWriteBuffer* RetiredHead = nullptr;
 	FWriteBuffer* RetiredTail = nullptr;
 
+	static const uint32 BufferSize = 8192 - 64; // "-64" is to allow overflow/encoding overhead
 	struct FCollector
 	{
-		struct FPacket
-		{
-			struct FHeader
-			{
-				uint16	Serial;
-				uint16	Size; // including header
-			};
-			FHeader		Header;
-			uint8		Data[8192];
-		};
-
 		void Send(const void* Data, uint32 Size)
 		{
 			if (GDataState == EDataState::Sending)
@@ -392,23 +387,45 @@ static void Writer_ConsumeEvents()
 
 		void Flush()
 		{
-			if (Cursor == sizeof(FPacket::Data))
+			if (Cursor == sizeof(Buffer))
 			{
 				return;
 			}
 
-			// There will always be space remaining for the header because we've
-			// arranged for that by including the header in FPacket. We'll shift
-			// it forward so it butts up against the event data (at the expense of
-			// some occasional unaligned stores).
-			Cursor -= sizeof(FPacket::FHeader);
+			struct FPacketBase
+			{
+				uint16	Serial;
+				uint16	PacketSize;
+			};
+
+#if 0
+			struct FPacketEncoded
+				: public FPacketBase
+			{
+				uint16	DecodedSize;
+			};
+
+			struct FPacket
+				: public FPacketEncoded
+			{
+				uint8 Data[BufferSize + 64];
+			};
 
-			auto* Out = (FPacket::FHeader*)(Packet.Data + Cursor);
-			Out->Serial = Serial;
-			Out->Size = sizeof(FPacket::Data) - Cursor;
-			Send(Out, Out->Size);
+			FPacket Packet;
+			Packet.Serial = Serial | 0x8000; // MSB indicates the packet is encoded
+			Packet.DecodedSize = uint16(sizeof(Buffer) - Cursor);
+			Packet.PacketSize = Encode(Buffer + Cursor, Packet.DecodedSize, Packet.Data, sizeof(Packet.Data));
+			Packet.PacketSize += sizeof(FPacketEncoded);
+			Send(&Packet, Packet.PacketSize);
+#else
+			Cursor -= sizeof(FPacketBase);
+			auto* Packet = (FPacketBase*)(Buffer + Cursor);
+			Packet->Serial = Serial;
+			Packet->PacketSize = uint16(sizeof(Buffer) - Cursor);
+			Send(Packet, Packet->PacketSize);
+#endif
 
-			Cursor = sizeof(FPacket::Data);
+			Cursor = sizeof(Buffer);
 			Serial++;
 		}
 
@@ -420,12 +437,13 @@ static void Writer_ConsumeEvents()
 			}
 
 			Cursor -= Size;
-			memcpy(Packet.Data + Cursor, Data, Size);
+			memcpy(Buffer + Cursor, Data, Size);
 		}
 
-		int16		Cursor = sizeof(FPacket::Data);
-		uint16		Serial = 0;
-		FPacket		Packet;
+		int16	Cursor = sizeof(Buffer);
+		uint16	Serial = 0;
+		uint8	Slack[8];				// Overflow for "Cursor - FPacketBase"
+		uint8	Buffer[BufferSize];
 	};
 
 	FCollector Collector;
