Renamed FPayloadTransort to FPacketTransport

#rnx

diff --git a/Engine/Source/Developer/TraceAnalysis/Private/Analysis/Engine.cpp b/Engine/Source/Developer/TraceAnalysis/Private/Analysis/Engine.cpp
--- a/Engine/Source/Developer/TraceAnalysis/Private/Analysis/Engine.cpp
+++ b/Engine/Source/Developer/TraceAnalysis/Private/Analysis/Engine.cpp
@@ -642,7 +642,7 @@ bool FAnalysisEngine::EstablishTransport(FStreamReader::FData& Data)
 	switch (Header->Format)
 	{
 	case 1:		Transport = new FTransport(); break;
-	case 2:		Transport = new FPayloadTransport(); break;
+	case 2:		Transport = new FPacketTransport(); break;
 	default:	return false;
 	//case 'E':	/* See the magic above */ break;
 	//case 'T':	/* See the magic above */ break;
diff --git a/Engine/Source/Developer/TraceAnalysis/Private/Analysis/Transport/PacketTransport.cpp b/Engine/Source/Developer/TraceAnalysis/Private/Analysis/Transport/PacketTransport.cpp
--- a/Engine/Source/Developer/TraceAnalysis/Private/Analysis/Transport/PacketTransport.cpp
+++ b/Engine/Source/Developer/TraceAnalysis/Private/Analysis/Transport/PacketTransport.cpp
@@ -16,7 +16,7 @@ TRACELOG_API int32 Decode(const void*, int32, void*, int32);
 
 
 ////////////////////////////////////////////////////////////////////////////////
-struct FPayloadTransport::FPacketNode
+struct FPacketTransport::FPacketNode
 {
 	FPacketNode*		Next;
 	uint32				Cursor;
@@ -25,8 +25,10 @@ struct FPayloadTransport::FPacketNode
 	uint8				Data[];
 };
 
+
+
 ////////////////////////////////////////////////////////////////////////////////
-FPayloadTransport::~FPayloadTransport()
+FPacketTransport::~FPacketTransport()
 {
 	for (FPacketNode* Root : {ActiveList, PendingList, FreeList})
 	{
@@ -40,7 +42,7 @@ FPayloadTransport::~FPayloadTransport()
 }
 
 ////////////////////////////////////////////////////////////////////////////////
-void FPayloadTransport::Advance(uint32 BlockSize)
+void FPacketTransport::Advance(uint32 BlockSize)
 {
 	if (ActiveList != nullptr)
 	{
@@ -49,7 +51,7 @@ void FPayloadTransport::Advance(uint32 BlockSize)
 }
 
 ////////////////////////////////////////////////////////////////////////////////
-const uint8* FPayloadTransport::GetPointerImpl(uint32 BlockSize)
+const uint8* FPacketTransport::GetPointerImpl(uint32 BlockSize)
 {
 	if (ActiveList == nullptr && !GetNextBatch())
 	{
@@ -70,7 +72,7 @@ const uint8* FPayloadTransport::GetPointerImpl(uint32 BlockSize)
 }
 
 ////////////////////////////////////////////////////////////////////////////////
-FPayloadTransport::FPacketNode* FPayloadTransport::AllocateNode()
+FPacketTransport::FPacketNode* FPacketTransport::AllocateNode()
 {
 	FPacketNode* Node;
 	if (FreeList != nullptr)
@@ -88,7 +90,7 @@ FPayloadTransport::FPacketNode* FPayloadTransport::AllocateNode()
 }
 
 ////////////////////////////////////////////////////////////////////////////////
-bool FPayloadTransport::GetNextBatch()
+bool FPacketTransport::GetNextBatch()
 {
 	int16 LastSerial = -1;
 	if (PendingList != nullptr)
diff --git a/Engine/Source/Developer/TraceAnalysis/Private/Analysis/Transport/PacketTransport.h b/Engine/Source/Developer/TraceAnalysis/Private/Analysis/Transport/PacketTransport.h
--- a/Engine/Source/Developer/TraceAnalysis/Private/Analysis/Transport/PacketTransport.h
+++ b/Engine/Source/Developer/TraceAnalysis/Private/Analysis/Transport/PacketTransport.h
@@ -8,11 +8,11 @@ namespace Trace
 {
 
 ////////////////////////////////////////////////////////////////////////////////
-class FPayloadTransport
+class FPacketTransport
 	: public FTransport
 {
 public:
-	virtual					~FPayloadTransport();
+	virtual					~FPacketTransport();
 	virtual void			Advance(uint32 BlockSize) override;
 	virtual const uint8*	GetPointerImpl(uint32 BlockSize) override;
 
