Added missing dependency on TraceLog from TraceAnalysis which currently includes from TraceLog and will soon link to it too.

#rnx

diff --git a/Engine/Source/Developer/TraceAnalysis/TraceAnalysis.Build.cs b/Engine/Source/Developer/TraceAnalysis/TraceAnalysis.Build.cs
--- a/Engine/Source/Developer/TraceAnalysis/TraceAnalysis.Build.cs
+++ b/Engine/Source/Developer/TraceAnalysis/TraceAnalysis.Build.cs
@@ -7,7 +7,8 @@ public class TraceAnalysis : ModuleRules
 	public TraceAnalysis(ReadOnlyTargetRules Target) : base(Target)
 	{
 		PrivateDependencyModuleNames.Add("Core");
-		PrivateDependencyModuleNames.Add("Sockets");
 		PrivateDependencyModuleNames.Add("DirectoryWatcher");
+		PrivateDependencyModuleNames.Add("Sockets");
+		PrivateDependencyModuleNames.Add("TraceLog");
 	}
 }
