Removed unused FDispatch member. Explicitly calling out available padding is not necessary, and in this instance, completely wrong.

#rnx

diff --git a/Engine/Source/Developer/TraceAnalysis/Private/Analysis/Engine.cpp b/Engine/Source/Developer/TraceAnalysis/Private/Analysis/Engine.cpp
--- a/Engine/Source/Developer/TraceAnalysis/Private/Analysis/Engine.cpp
+++ b/Engine/Source/Developer/TraceAnalysis/Private/Analysis/Engine.cpp
@@ -51,7 +51,6 @@ struct FAnalysisEngine::FDispatch
 	uint16			EventSize			= 0;
 	uint16			LoggerNameOffset	= 0;	// From FDispatch ptr
 	uint16			EventNameOffset		= 0;	// From FDispatch ptr
-	uint16			_Unused0;
 	FField			Fields[];
 };
 
