Consolidated everything related to the trace event "protocol" (headers, new-event event's format) into one place and wrapped it in a namespace so it is versioned. Updated analysis engine to support multiple protocols (of which there is currently exactly one).

#rnx

diff --git a/Engine/Source/Developer/TraceAnalysis/Private/Analysis/Engine.cpp b/Engine/Source/Developer/TraceAnalysis/Private/Analysis/Engine.cpp
--- a/Engine/Source/Developer/TraceAnalysis/Private/Analysis/Engine.cpp
+++ b/Engine/Source/Developer/TraceAnalysis/Private/Analysis/Engine.cpp
@@ -6,7 +6,7 @@
 #include "HAL/UnrealMemory.h"
 #include "Trace/Analysis.h"
 #include "Trace/Analyzer.h"
-#include "Trace/Detail/EventDef.h"
+#include "Trace/Detail/Protocol.h"
 #include "Transport/Transport.h"
 #include "Transport/PacketTransport.h"
 
@@ -353,11 +353,6 @@
 	AddRoute(SelfIndex, RouteId_NewEvent, RouteHash_NewEvent);
 	AddRoute(SelfIndex, RouteId_NewTrace, "$Trace", "NewTrace");
 	AddRoute(SelfIndex, RouteId_Timing, "$Trace", "Timing");
-
-	// Add a dispatch for special events
-	FDispatchBuilder Builder;
-	Builder.SetUid(uint16(FNewEventEvent::Uid));
-	AddDispatch(Builder.Finalize());
 }
 
 ////////////////////////////////////////////////////////////////////////////////
@@ -396,18 +391,5 @@
 }
 
 ////////////////////////////////////////////////////////////////////////////////
-bool FAnalysisEngine::OnEvent(uint16 RouteId, const FOnEventContext& Context)
-{
-	switch (RouteId)
-	{
-	case RouteId_NewEvent:	OnNewEventInternal(Context);	break;
-	case RouteId_NewTrace:	OnNewTrace(Context);			break;
-	case RouteId_Timing:	OnTiming(Context);				break;
-	}
-
-	return true;
-}
-
-////////////////////////////////////////////////////////////////////////////////
 void FAnalysisEngine::AddRoute(
 	uint16 AnalyzerIndex,
@@ -437,5 +419,18 @@
 }
 
 ////////////////////////////////////////////////////////////////////////////////
+bool FAnalysisEngine::OnEvent(uint16 RouteId, const FOnEventContext& Context)
+{
+	switch (RouteId)
+	{
+	case RouteId_NewEvent:	OnNewEventInternal(Context);	break;
+	case RouteId_NewTrace:	OnNewTrace(Context);			break;
+	case RouteId_Timing:	OnTiming(Context);				break;
+	}
+
+	return true;
+}
+
+////////////////////////////////////////////////////////////////////////////////
 void FAnalysisEngine::OnNewTrace(const FOnEventContext& Context)
 {
@@ -525,35 +520,11 @@
 void FAnalysisEngine::OnNewEventInternal(const FOnEventContext& Context)
 {
 	const FEventDataInfo& EventData = (const FEventDataInfo&)(Context.EventData);
-	const auto& NewEvent = *(FNewEventEvent*)(EventData.Ptr);
 
 	FDispatchBuilder Builder;
-
-	const auto* NameCursor = (const ANSICHAR*)(NewEvent.Fields + NewEvent.FieldCount);
-
-	Builder.SetLoggerName(NameCursor, NewEvent.LoggerNameSize);
-	NameCursor += NewEvent.LoggerNameSize;
-
-	Builder.SetEventName(NameCursor, NewEvent.EventNameSize);
-	NameCursor += NewEvent.EventNameSize;
-	Builder.SetUid(NewEvent.EventUid);
-
-	// Fill out the fields
-	for (int i = 0, n = NewEvent.FieldCount; i < n; ++i)
+	if (ProtocolHandler == &FAnalysisEngine::OnDataProtocol0)
 	{
-		const auto& Field = NewEvent.Fields[i];
-
-		uint16 TypeSize = 1 << (Field.TypeInfo & _Field_Pow2SizeMask);
-
-		auto TypeId = FEventFieldInfo::EType::Integer;
-		if ((Field.TypeInfo & _Field_CategoryMask) == _Field_Float)
-		{
-			TypeId = FEventFieldInfo::EType::Float;
-		}
-
-		Builder.AddField(NameCursor, Field.NameSize, Field.Offset, Field.Size, TypeId, TypeSize);
-
-		NameCursor += Field.NameSize;
+		OnNewEventProtocol0(Builder, EventData.Ptr);
 	}
 
 	// Get the dispatch and add it into the dispatch table. Fail gently if there
@@ -577,6 +548,8 @@
 ////////////////////////////////////////////////////////////////////////////////
 bool FAnalysisEngine::AddDispatch(FDispatch* Dispatch)
 {
+	// Add the dispatch to the dispatch table, failing gently if the table slot
+	// is already occupied.
 	uint16 Uid = Dispatch->Uid;
 	if (Uid < Dispatches.Num())
  	{
@@ -591,8 +564,6 @@
 		Dispatches.SetNum(Uid + 1);
  	}
 
-	Dispatches[Uid] = Dispatch;
-
 	// Find routes that have subscribed to this event.
 	for (uint16 i = 0, n = Routes.Num(); i < n; ++i)
 	{
@@ -603,15 +574,49 @@
 		}
 	}
 
+	Dispatches[Uid] = Dispatch;
 	return true;
 }
 
 ////////////////////////////////////////////////////////////////////////////////
+void FAnalysisEngine::OnNewEventProtocol0(FDispatchBuilder& Builder, const void* EventData)
+{
+	const auto& NewEvent = *(Protocol0::FNewEventEvent*)(EventData);
+
+	const auto* NameCursor = (const ANSICHAR*)(NewEvent.Fields + NewEvent.FieldCount);
+
+	Builder.SetLoggerName(NameCursor, NewEvent.LoggerNameSize);
+	NameCursor += NewEvent.LoggerNameSize;
+
+	Builder.SetEventName(NameCursor, NewEvent.EventNameSize);
+	NameCursor += NewEvent.EventNameSize;
+	Builder.SetUid(NewEvent.EventUid);
+
+	// Fill out the fields
+	for (int i = 0, n = NewEvent.FieldCount; i < n; ++i)
+	{
+		const auto& Field = NewEvent.Fields[i];
+
+		uint16 TypeSize = 1 << (Field.TypeInfo & Protocol0::Field_Pow2SizeMask);
+
+		auto TypeId = FEventFieldInfo::EType::Integer;
+		if ((Field.TypeInfo & Protocol0::Field_CategoryMask) == Protocol0::Field_Float)
+		{
+			TypeId = FEventFieldInfo::EType::Float;
+		}
+
+		Builder.AddField(NameCursor, Field.NameSize, Field.Offset, Field.Size, TypeId, TypeSize);
+
+		NameCursor += Field.NameSize;
+	}
+}
+
+////////////////////////////////////////////////////////////////////////////////
 bool FAnalysisEngine::EstablishTransport(FStreamReader::FData& Data)
 {
 	const struct {
-		uint8 Format;
-		uint8 Parameter;
+		uint8 TransportVersion;
+		uint8 ProtocolVersion;
 	}* Header = decltype(Header)(Data.GetPointer(sizeof(*Header)));
 	if (Header == nullptr)
 	{
@@ -620,7 +625,7 @@
 
 	// Check for the magic uint32. Early traces did not include this as it was
 	// used to validate a inbound socket connection and then discarded.
-	if (Header->Format == 'E' || Header->Format == 'T')
+	if (Header->TransportVersion == 'E' || Header->TransportVersion == 'T')
 	{
 		const uint32* Magic = (const uint32*)(Data.GetPointer(sizeof(*Magic)));
 		if (*Magic == 'ECRT')
@@ -638,7 +643,7 @@
 		return false;
 	}
 
-	switch (Header->Format)
+	switch (Header->TransportVersion)
 	{
 	case 1:		Transport = new FTransport(); break;
 	case 2:		Transport = new FPacketTransport(); break;
@@ -647,6 +652,21 @@
 	//case 'T':	/* See the magic above */ break;
 	}
 
+	switch (Header->ProtocolVersion)
+	{
+	case 0:
+		ProtocolHandler = &FAnalysisEngine::OnDataProtocol0;
+		{
+			FDispatchBuilder Builder;
+			Builder.SetUid(uint16(Protocol0::FNewEventEvent::Uid));
+			AddDispatch(Builder.Finalize());
+		}
+		break;
+
+	default:
+		return false;
+	}
+
 	Data.Advance(sizeof(*Header));
 	return true;
 }
@@ -668,25 +688,35 @@
 		}
 	}
 
-	struct FEventHeader
-	{
-		uint16	Uid;
-		uint16	Size;
-		uint8	EventData[];
-	};
-
 	Transport->SetSource(Data);
+	bool Ret = (this->*ProtocolHandler)();
 
+	// If there's no analyzers left we might as well not continue
+	int32 ActiveAnalyzerCount = 0;
+	for (IAnalyzer* Analyzer : Analyzers)
+	{
+		if ((Analyzer != nullptr) && (Analyzer != this))
+		{
+			ActiveAnalyzerCount++;
+		}
+	}
+
+	return (ActiveAnalyzerCount > 0) & Ret;
+}
+
+////////////////////////////////////////////////////////////////////////////////
+bool FAnalysisEngine::OnDataProtocol0()
+{
 	while (true)
 	{
-		const auto* Header = Transport->GetPointer<FEventHeader>();
+		const auto* Header = Transport->GetPointer<Protocol0::FEventHeader>();
 		if (Header == nullptr)
 		{
 			break;
 		}
 
-		uint32 BlockSize = Header->Size + sizeof(FEventHeader);
-		Header = Transport->GetPointer<FEventHeader>(BlockSize);
+		uint32 BlockSize = Header->Size + sizeof(Protocol0::FEventHeader);
+		Header = Transport->GetPointer<Protocol0::FEventHeader>(BlockSize);
 		if (Header == nullptr)
 		{
 			break;
@@ -718,17 +748,7 @@
 		});
 	}
 
-	// If there's no analyzers left we might as well not continue
-	int32 ActiveAnalyzerCount = 0;
-	for (IAnalyzer* Analyzer : Analyzers)
-	{
-		if ((Analyzer != nullptr) && (Analyzer != this))
-		{
-			ActiveAnalyzerCount++;
-		}
-	}
-
-	return (ActiveAnalyzerCount > 0);
+	return true;
 }
 
 } // namespace Trace
diff --git a/Engine/Source/Developer/TraceAnalysis/Private/Analysis/Engine.h b/Engine/Source/Developer/TraceAnalysis/Private/Analysis/Engine.h
--- a/Engine/Source/Developer/TraceAnalysis/Private/Analysis/Engine.h
+++ b/Engine/Source/Developer/TraceAnalysis/Private/Analysis/Engine.h
@@ -22,6 +22,8 @@ public:
 	bool				OnData(FStreamReader::FData& Data);
 
 private:
+	typedef bool (FAnalysisEngine::*ProtocolHandlerType)();
+
 	struct FRoute
 	{
 		uint32			Hash;
@@ -33,22 +35,25 @@ private:
 
 	class				FDispatchBuilder;
 	virtual bool		OnEvent(uint16 RouteId, const FOnEventContext& Context) override;
+	void				OnNewTrace(const FOnEventContext& Context);
+	void				OnTiming(const FOnEventContext& Context);
+	void				OnNewEventInternal(const FOnEventContext& Context);
+	void				OnNewEventProtocol0(FDispatchBuilder& Builder, const void* EventData);
 
 	bool				EstablishTransport(FStreamReader::FData& Data);
+	bool				OnDataProtocol0();
 	bool				AddDispatch(FDispatch* Dispatch);
 	template <typename ImplType>
 	void				ForEachRoute(const FDispatch* Dispatch, ImplType&& Impl);
 	void				AddRoute(uint16 AnalyzerIndex, uint16 Id, const ANSICHAR* Logger, const ANSICHAR* Event);
 	void				AddRoute(uint16 AnalyzerIndex, uint16 Id, uint32 Hash);
-	void				OnNewTrace(const FOnEventContext& Context);
-	void				OnTiming(const FOnEventContext& Context);
-	void				OnNewEventInternal(const FOnEventContext& Context);
 	void				RetireAnalyzer(IAnalyzer* Analyzer);
 	FSessionContext		SessionContext;
 	TArray<FRoute>		Routes;
 	TArray<IAnalyzer*>	Analyzers;
 	TArray<FDispatch*>	Dispatches;
 	class FTransport*	Transport = nullptr;
+	ProtocolHandlerType	ProtocolHandler = nullptr;
 };
 
 } // namespace Trace
diff --git a/Engine/Source/Runtime/TraceLog/Private/Trace/Writer.cpp b/Engine/Source/Runtime/TraceLog/Private/Trace/Writer.cpp
--- a/Engine/Source/Runtime/TraceLog/Private/Trace/Writer.cpp
+++ b/Engine/Source/Runtime/TraceLog/Private/Trace/Writer.cpp
@@ -6,6 +6,7 @@
 
 #include "Trace/Detail/Atomic.h"
 #include "Trace/Platform.h"
+#include "Trace/Detail/Protocol.h"
 #include "Trace/Trace.h"
 
 #include "Misc/CString.h"
@@ -508,9 +509,9 @@ static void Writer_UpdateData()
 
 		// Stream header
 		const struct {
-			uint8 Format;
-			uint8 Parameter;
-		} TransportHeader = { 2 };
+			uint8 TransportVersion	= 4;
+			uint8 ProtocolVersion	= 0;
+		} TransportHeader;
 		bOk &= IoWrite(GDataHandle, &TransportHeader, sizeof(TransportHeader));
 
 		// Passively collected data
diff --git a/Engine/Source/Runtime/TraceLog/Public/Trace/Detail/EventDef.h b/Engine/Source/Runtime/TraceLog/Public/Trace/Detail/EventDef.h
--- a/Engine/Source/Runtime/TraceLog/Public/Trace/Detail/EventDef.h
+++ b/Engine/Source/Runtime/TraceLog/Public/Trace/Detail/EventDef.h
@@ -4,41 +4,6 @@
 
 #include "Trace/Config.h"
 
-namespace Trace
-{
-
-#if defined(_MSC_VER)
-	#pragma warning(push)
-	#pragma warning(disable : 4200) // non-standard zero-sized array
-#endif
-
-////////////////////////////////////////////////////////////////////////////////
-struct FNewEventEvent
-{
-	enum : uint16 { Uid = 0 };
-
-	uint16		EventUid;
-	uint16		FieldCount;
-	uint8		LoggerNameSize;
-	uint8		EventNameSize;
-	struct
-	{
-		uint16	Offset;
-		uint16	Size;
-		uint8	TypeInfo;
-		uint8	NameSize;
-	}			Fields[];
-	/*uint8		NameData[]*/
-};
-
-#if defined(_MSC_VER)
-	#pragma warning(pop)
-#endif
-
-} // namespace Trace
-
-
-
 #if UE_TRACE_ENABLED
 
 namespace Trace
diff --git a/Engine/Source/Runtime/TraceLog/Public/Trace/Detail/Field.h b/Engine/Source/Runtime/TraceLog/Public/Trace/Detail/Field.h
--- a/Engine/Source/Runtime/TraceLog/Public/Trace/Detail/Field.h
+++ b/Engine/Source/Runtime/TraceLog/Public/Trace/Detail/Field.h
@@ -4,52 +4,14 @@
 
 #include "Trace/Config.h"
 
+#if UE_TRACE_ENABLED
+
+#include "Protocol.h"
 #include "Templates/UnrealTemplate.h"
 
 namespace Trace
 {
 
-////////////////////////////////////////////////////////////////////////////////
-enum : uint8
-{
-	/* Category */
-	_Field_CategoryMask	= 0300,
-	_Field_Integer		= 0000, 
-	_Field_Float		= 0100,
-	_Field_Array		= 0200, /* flag */
-
-	/* Size */
-	_Field_Pow2SizeMask	= 0003,
-	_Field_8			= 0000,
-	_Field_16			= 0001,
-	_Field_32			= 0002,
-	_Field_64			= 0003,
-#if PLATFORM_64BITS
-	_Field_Ptr			= _Field_64,
-#else
-	_Field_Ptr			= _Field_32,
-#endif
-
-	/* unused			= 0004, flag */
-	/* unused			= 0070, enum/flags*/
-};
-
-////////////////////////////////////////////////////////////////////////////////
-enum EFieldType : uint8
-{
-	Field_Bool		= _Field_Integer | _Field_8,
-	Field_Int8		= _Field_Integer | _Field_8,
-	Field_Int16		= _Field_Integer | _Field_16,
-	Field_Int32		= _Field_Integer | _Field_32,
-	Field_Int64		= _Field_Integer | _Field_64,
-	Field_Ptr		= _Field_Integer | _Field_Ptr,
-	Field_Float		= _Field_Float	 | _Field_32,
-	Field_Double	= _Field_Float	 | _Field_64,
-
-	/* Flags */
-	Field_Array		= _Field_Array,
-};
-
 ////////////////////////////////////////////////////////////////////////////////
 template <typename Type> struct TFieldType;
 
@@ -62,16 +24,9 @@ template <> struct TFieldType<uint8>		{ enum { Value = Field_Int8 }; };
 template <> struct TFieldType<uint16>		{ enum { Value = Field_Int16 }; };
 template <> struct TFieldType<uint32>		{ enum { Value = Field_Int32 }; };
 template <> struct TFieldType<uint64>		{ enum { Value = Field_Int64 }; };
-template <> struct TFieldType<float>		{ enum { Value = Field_Float }; };
-template <> struct TFieldType<double>		{ enum { Value = Field_Double }; };
-template <class T> struct TFieldType<T*>	{ enum { Value = Field_Ptr }; };
-
-} // namespace Trace
-
-#if UE_TRACE_ENABLED
-
-namespace Trace
-{
+template <> struct TFieldType<float>		{ enum { Value = Field_Float32 }; };
+template <> struct TFieldType<double>		{ enum { Value = Field_Float64 }; };
+template <class T> struct TFieldType<T*>	{ enum { Value = Field_Pointer }; };
 
 ////////////////////////////////////////////////////////////////////////////////
 struct FLiteralName
diff --git a/Engine/Source/Runtime/TraceLog/Public/Trace/Detail/Protocol.h b/Engine/Source/Runtime/TraceLog/Public/Trace/Detail/Protocol.h
new file mode 100644
--- /dev/null
+++ b/Engine/Source/Runtime/TraceLog/Public/Trace/Detail/Protocol.h
@@ -0,0 +1,14 @@
+// Copyright 1998-2019 Epic Games, Inc. All Rights Reserved.
+
+#pragma once
+
+#if defined(_MSC_VER)
+	#pragma warning(push)
+	#pragma warning(disable : 4200) // non-standard zero-sized array
+#endif
+
+#include "Protocols/Protocol0.h"
+
+#if defined(_MSC_VER)
+	#pragma warning(pop)
+#endif
diff --git a/Engine/Source/Runtime/TraceLog/Public/Trace/Detail/Protocols/Protocol0.h b/Engine/Source/Runtime/TraceLog/Public/Trace/Detail/Protocols/Protocol0.h
new file mode 100644
--- /dev/null
+++ b/Engine/Source/Runtime/TraceLog/Public/Trace/Detail/Protocols/Protocol0.h
@@ -0,0 +1,77 @@
+// Copyright 1998-2019 Epic Games, Inc. All Rights Reserved.
+
+#pragma once
+
+namespace Trace
+{
+
+#if defined(TRACE_PRIVATE_PROTOCOL_0)
+inline
+#endif
+namespace Protocol0
+{
+
+////////////////////////////////////////////////////////////////////////////////
+enum : uint8
+{
+	/* Category */
+	Field_CategoryMask	= 0300,
+	Field_Integer		= 0000, 
+	Field_Float			= 0100,
+
+	/* Size */
+	Field_Pow2SizeMask	= 0003,
+	Field_8				= 0000,
+	Field_16			= 0001,
+	Field_32			= 0002,
+	Field_64			= 0003,
+#if PLATFORM_64BITS
+	Field_Ptr			= Field_64,
+#else
+	Field_Ptr			= Field_32,
+#endif
+};
+
+////////////////////////////////////////////////////////////////////////////////
+enum EFieldType : uint8
+{
+	Field_Bool		= Field_Integer | Field_8,
+	Field_Int8		= Field_Integer | Field_8,
+	Field_Int16		= Field_Integer | Field_16,
+	Field_Int32		= Field_Integer | Field_32,
+	Field_Int64		= Field_Integer | Field_64,
+	Field_Pointer	= Field_Integer | Field_Ptr,
+	Field_Float32	= Field_Float   | Field_32,
+	Field_Float64	= Field_Float   | Field_64,
+};
+
+////////////////////////////////////////////////////////////////////////////////
+struct FNewEventEvent
+{
+	enum : uint16 { Uid = 0 };
+
+	uint16		EventUid;
+	uint16		FieldCount;
+	uint8		LoggerNameSize;
+	uint8		EventNameSize;
+	struct
+	{
+		uint16	Offset;
+		uint16	Size;
+		uint8	TypeInfo;
+		uint8	NameSize;
+	}			Fields[];
+	/*uint8		NameData[]*/
+};
+
+////////////////////////////////////////////////////////////////////////////////
+struct FEventHeader
+{
+	uint16		Uid;
+	uint16		Size;
+	uint8		EventData[];
+};
+
+} // namespace Protocol0
+
+} // namespace Trace
diff --git a/Engine/Source/Runtime/TraceLog/Public/Trace/Detail/Trace.h b/Engine/Source/Runtime/TraceLog/Public/Trace/Detail/Trace.h
--- a/Engine/Source/Runtime/TraceLog/Public/Trace/Detail/Trace.h
+++ b/Engine/Source/Runtime/TraceLog/Public/Trace/Detail/Trace.h
@@ -7,6 +7,8 @@
 
 #if UE_TRACE_ENABLED
 
+#define TRACE_PRIVATE_PROTOCOL_0
+
 #define TRACE_PRIVATE_EVENT_DEFINE(LoggerName, EventName) \
 	Trace::FEventDef LoggerName##EventName##Event;
 
