Fixed a couple of rule breaks

#rnx

diff --git a/Engine/Source/Developer/TraceAnalysis/Private/Analysis/Engine.cpp b/Engine/Source/Developer/TraceAnalysis/Private/Analysis/Engine.cpp
--- a/Engine/Source/Developer/TraceAnalysis/Private/Analysis/Engine.cpp
+++ b/Engine/Source/Developer/TraceAnalysis/Private/Analysis/Engine.cpp
@@ -689,7 +689,7 @@ bool FAnalysisEngine::OnData(FStreamReader::FData& Data)
 	}
 
 	Transport->SetSource(Data);
-	bool Ret = (this->*ProtocolHandler)();
+	bool bRet = (this->*ProtocolHandler)();
 
 	// If there's no analyzers left we might as well not continue
 	int32 ActiveAnalyzerCount = 0;
@@ -701,7 +701,7 @@ bool FAnalysisEngine::OnData(FStreamReader::FData& Data)
 		}
 	}
 
-	return (ActiveAnalyzerCount > 0) & Ret;
+	return (ActiveAnalyzerCount > 0) & bRet;
 }
 
 ////////////////////////////////////////////////////////////////////////////////
