Only sort routes that were added by analyzers. As dispatches link to routes by index this ensures we don't invalidate any existing dispatches. It also help to hide internals as writers of analyzers won't be able to route events internal to the operation of Trace and analysis.

#rnx

diff --git a/Engine/Source/Developer/TraceAnalysis/Private/Analysis/Engine.cpp b/Engine/Source/Developer/TraceAnalysis/Private/Analysis/Engine.cpp
--- a/Engine/Source/Developer/TraceAnalysis/Private/Analysis/Engine.cpp
+++ b/Engine/Source/Developer/TraceAnalysis/Private/Analysis/Engine.cpp
@@ -451,6 +451,10 @@ void FAnalysisEngine::OnNewTrace(const FOnEventContext& Context)
 	} Builder;
 	Builder.Self = this;
 
+	// Some internal routes have been established already. In case there's some
+	// dispatches that are already connected to these routes we won't sort them
+	uint32 FixedRouteCount = Routes.Num();
+
 	FOnAnalysisContext OnAnalysisContext = { { SessionContext }, Builder };
 	for (uint16 i = 0, n = Analyzers.Num(); i < n; ++i)
 	{
@@ -458,7 +462,8 @@ void FAnalysisEngine::OnNewTrace(const FOnEventContext& Context)
 		Analyzers[i]->OnAnalysisBegin(OnAnalysisContext);
 	}
 
-	Algo::SortBy(Routes, [] (const FRoute& Route) { return Route.Hash; });
+	TArrayView<FRoute> RouteSubset(Routes.GetData() + FixedRouteCount, Routes.Num() - FixedRouteCount);
+	Algo::SortBy(RouteSubset, [] (const FRoute& Route) { return Route.Hash; });
 
 	FRoute* Cursor = Routes.GetData();
 	Cursor->Count = 1;
