Allow the dispatch builder's SetLoggerName/SetEventName to be used without knowing the string length ahead of time

#rnx

diff --git a/Engine/Source/Developer/TraceAnalysis/Private/Analysis/Engine.cpp b/Engine/Source/Developer/TraceAnalysis/Private/Analysis/Engine.cpp
--- a/Engine/Source/Developer/TraceAnalysis/Private/Analysis/Engine.cpp
+++ b/Engine/Source/Developer/TraceAnalysis/Private/Analysis/Engine.cpp
@@ -62,8 +62,8 @@ class FAnalysisEngine::FDispatchBuilder
 public:
 					FDispatchBuilder();
 	void			SetUid(uint16 Uid);
-	void			SetLoggerName(const ANSICHAR* Name, int32 NameSize);
-	void			SetEventName(const ANSICHAR* Name, int32 NameSize);
+	void			SetLoggerName(const ANSICHAR* Name, int32 NameSize=-1);
+	void			SetEventName(const ANSICHAR* Name, int32 NameSize=-1);
 	void			AddField(const ANSICHAR* Name, int32 NameSize, uint16 Offset, uint16 Size, FEventFieldInfo::EType TypeId, uint16 TypeSize);
 	FDispatch*		Finalize();
 
