Gave the transport and protocol magic numbers some names so it is easier to follow

#rnx

diff --git a/Engine/Source/Developer/TraceAnalysis/Private/Analysis/Engine.cpp b/Engine/Source/Developer/TraceAnalysis/Private/Analysis/Engine.cpp
--- a/Engine/Source/Developer/TraceAnalysis/Private/Analysis/Engine.cpp
+++ b/Engine/Source/Developer/TraceAnalysis/Private/Analysis/Engine.cpp
@@ -650,16 +650,16 @@ bool FAnalysisEngine::EstablishTransport(FStreamReader::FData& Data)
 
 	switch (Header->TransportVersion)
 	{
-	case 1:		Transport = new FTransport(); break;
-	case 2:		Transport = new FPacketTransport(); break;
-	default:	return false;
+	case ETransport::Raw:		Transport = new FTransport(); break;
+	case ETransport::Packet:	Transport = new FPacketTransport(); break;
+	default:					return false;
 	//case 'E':	/* See the magic above */ break;
 	//case 'T':	/* See the magic above */ break;
 	}
 
 	switch (Header->ProtocolVersion)
 	{
-	case 0:
+	case Protocol0::EProtocol::Id:
 		ProtocolHandler = &FAnalysisEngine::OnDataProtocol0;
 		{
 			FDispatchBuilder Builder;
diff --git a/Engine/Source/Runtime/TraceLog/Private/Trace/Writer.cpp b/Engine/Source/Runtime/TraceLog/Private/Trace/Writer.cpp
--- a/Engine/Source/Runtime/TraceLog/Private/Trace/Writer.cpp
+++ b/Engine/Source/Runtime/TraceLog/Private/Trace/Writer.cpp
@@ -517,8 +517,8 @@ static void Writer_UpdateData()
 
 		// Stream header
 		const struct {
-			uint8 TransportVersion	= 2;
-			uint8 ProtocolVersion	= 0;
+			uint8 TransportVersion	= uint8(ETransport::Packet);
+			uint8 ProtocolVersion	= uint8(EProtocol::Id);
 		} TransportHeader;
 		bOk &= IoWrite(GDataHandle, &TransportHeader, sizeof(TransportHeader));
 
diff --git a/Engine/Source/Runtime/TraceLog/Public/Trace/Detail/Protocol.h b/Engine/Source/Runtime/TraceLog/Public/Trace/Detail/Protocol.h
--- a/Engine/Source/Runtime/TraceLog/Public/Trace/Detail/Protocol.h
+++ b/Engine/Source/Runtime/TraceLog/Public/Trace/Detail/Protocol.h
@@ -4,6 +4,20 @@
 
 #include "Trace/Config.h"
 
+////////////////////////////////////////////////////////////////////////////////
+namespace Trace
+{
+
+enum class ETransport : uint8
+{
+	_Unused		= 0,
+	Raw			= 1,
+	Packet		= 2,
+};
+
+} // namespace Trace
+
+////////////////////////////////////////////////////////////////////////////////
 #if UE_TRACE_ENABLED
 #	define TRACE_PRIVATE_PROTOCOL_0
 #endif
diff --git a/Engine/Source/Runtime/TraceLog/Public/Trace/Detail/Protocols/Protocol0.h b/Engine/Source/Runtime/TraceLog/Public/Trace/Detail/Protocols/Protocol0.h
--- a/Engine/Source/Runtime/TraceLog/Public/Trace/Detail/Protocols/Protocol0.h
+++ b/Engine/Source/Runtime/TraceLog/Public/Trace/Detail/Protocols/Protocol0.h
@@ -11,6 +11,9 @@ inline
 namespace Protocol0
 {
 
+////////////////////////////////////////////////////////////////////////////////
+enum class EProtocol : uint8 { Id = 0 };
+
 ////////////////////////////////////////////////////////////////////////////////
 enum : uint8
 {
