Fixed field name offset calculation that we really rather wrong.

#rnx

diff --git a/Engine/Source/Developer/TraceAnalysis/Private/Analysis/Engine.cpp b/Engine/Source/Developer/TraceAnalysis/Private/Analysis/Engine.cpp
--- a/Engine/Source/Developer/TraceAnalysis/Private/Analysis/Engine.cpp
+++ b/Engine/Source/Developer/TraceAnalysis/Private/Analysis/Engine.cpp
@@ -97,8 +97,8 @@ FAnalysisEngine::FDispatch* FAnalysisEngine::FDispatchBuilder::Finalize()
 	// Fix up name offsets
 	for (int i = 0, n = Dispatch->FieldCount; i < n; ++i)
 	{
-		auto& Field = Dispatch->Fields[0];
-		Field.NameOffset += sizeof(Field) * (n - i + 1);
+		auto* Field = Dispatch->Fields + i;
+		Field->NameOffset += Buffer.Num() - uint32(UPTRINT(Field) - UPTRINT(Dispatch));
 	}
 
 	// Calculate this dispatch's hash.
