Fixed Clang compile errors

#rnx

diff --git a/Engine/Source/Developer/TraceAnalysis/Private/Analysis/Engine.cpp b/Engine/Source/Developer/TraceAnalysis/Private/Analysis/Engine.cpp
--- a/Engine/Source/Developer/TraceAnalysis/Private/Analysis/Engine.cpp
+++ b/Engine/Source/Developer/TraceAnalysis/Private/Analysis/Engine.cpp
@@ -650,8 +650,8 @@ bool FAnalysisEngine::EstablishTransport(FStreamReader::FData& Data)
 
 	switch (Header->TransportVersion)
 	{
-	case ETransport::Raw:		Transport = new FTransport(); break;
-	case ETransport::Packet:	Transport = new FPacketTransport(); break;
+	case int(ETransport::Raw):		Transport = new FTransport(); break;
+	case int(ETransport::Packet):	Transport = new FPacketTransport(); break;
 	default:					return false;
 	//case 'E':	/* See the magic above */ break;
 	//case 'T':	/* See the magic above */ break;
@@ -659,7 +659,7 @@ bool FAnalysisEngine::EstablishTransport(FStreamReader::FData& Data)
 
 	switch (Header->ProtocolVersion)
 	{
-	case Protocol0::EProtocol::Id:
+	case int(Protocol0::EProtocol::Id):
 		ProtocolHandler = &FAnalysisEngine::OnDataProtocol0;
 		{
 			FDispatchBuilder Builder;
