Include the master trace version in the SessionContext object available to analyzers.

#rnx

diff --git a/Engine/Source/Developer/TraceAnalysis/Private/Analysis/Engine.cpp b/Engine/Source/Developer/TraceAnalysis/Private/Analysis/Engine.cpp
--- a/Engine/Source/Developer/TraceAnalysis/Private/Analysis/Engine.cpp
+++ b/Engine/Source/Developer/TraceAnalysis/Private/Analysis/Engine.cpp
@@ -434,6 +434,9 @@ bool FAnalysisEngine::OnEvent(uint16 RouteId, const FOnEventContext& Context)
 ////////////////////////////////////////////////////////////////////////////////
 void FAnalysisEngine::OnNewTrace(const FOnEventContext& Context)
 {
+	const FEventData& EventData = Context.EventData;
+	SessionContext.Version = EventData.GetValue<uint16>("Version");
+
 	struct : IAnalyzer::FInterfaceBuilder
 	{
 		virtual void RouteEvent(uint16 RouteId, const ANSICHAR* Logger, const ANSICHAR* Event) override
diff --git a/Engine/Source/Developer/TraceAnalysis/Public/Trace/Analyzer.h b/Engine/Source/Developer/TraceAnalysis/Public/Trace/Analyzer.h
--- a/Engine/Source/Developer/TraceAnalysis/Public/Trace/Analyzer.h
+++ b/Engine/Source/Developer/TraceAnalysis/Public/Trace/Analyzer.h
@@ -27,6 +27,7 @@ public:
 	{
 		uint64 StartCycle;
 		uint64 CycleFrequency;
+		uint16 Version;
 
 		double TimestampFromCycle(uint64 Cycle) const
 		{
