Retire analyzers immediately if they don't add any routes.

#rnx

diff --git a/Engine/Source/Developer/TraceAnalysis/Private/Analysis/Engine.cpp b/Engine/Source/Developer/TraceAnalysis/Private/Analysis/Engine.cpp
--- a/Engine/Source/Developer/TraceAnalysis/Private/Analysis/Engine.cpp
+++ b/Engine/Source/Developer/TraceAnalysis/Private/Analysis/Engine.cpp
@@ -461,8 +461,17 @@ void FAnalysisEngine::OnNewTrace(const FOnEventContext& Context)
 	FOnAnalysisContext OnAnalysisContext = { { SessionContext }, Builder };
 	for (uint16 i = 0, n = Analyzers.Num(); i < n; ++i)
 	{
+		uint32 RouteCount = Routes.Num();
+
 		Builder.AnalyzerIndex = i;
-		Analyzers[i]->OnAnalysisBegin(OnAnalysisContext);
+		IAnalyzer* Analyzer = Analyzers[i];
+		Analyzer->OnAnalysisBegin(OnAnalysisContext);
+
+		// If the analyzer didn't add any routes we'll retire it immediately
+		if (RouteCount == Routes.Num() && Analyzer != this)
+		{
+			RetireAnalyzer(Analyzer);
+		}
 	}
 
 	TArrayView<FRoute> RouteSubset(Routes.GetData() + FixedRouteCount, Routes.Num() - FixedRouteCount);
