Removed the inner FStreamReader::FData and folded it into FStreamReader. FStreamReader now also has its own home instead of crashing DataStreams' home.

#rnx

diff --git a/Engine/Source/Developer/TraceAnalysis/Private/Analysis/Engine.cpp b/Engine/Source/Developer/TraceAnalysis/Private/Analysis/Engine.cpp
--- a/Engine/Source/Developer/TraceAnalysis/Private/Analysis/Engine.cpp
+++ b/Engine/Source/Developer/TraceAnalysis/Private/Analysis/Engine.cpp
@@ -4,11 +4,12 @@
 #include "Containers/ArrayView.h"
 #include "CoreGlobals.h"
 #include "HAL/UnrealMemory.h"
+#include "StreamReader.h"
 #include "Trace/Analysis.h"
 #include "Trace/Analyzer.h"
 #include "Trace/Detail/Protocol.h"
-#include "Transport/Transport.h"
 #include "Transport/PacketTransport.h"
+#include "Transport/Transport.h"
 
 namespace Trace
 {
@@ -629,12 +630,12 @@ void FAnalysisEngine::OnNewEventProtocol0(FDispatchBuilder& Builder, const void*
 }
 
 ////////////////////////////////////////////////////////////////////////////////
-bool FAnalysisEngine::EstablishTransport(FStreamReader::FData& Data)
+bool FAnalysisEngine::EstablishTransport(FStreamReader& Reader)
 {
 	const struct {
 		uint8 TransportVersion;
 		uint8 ProtocolVersion;
-	}* Header = decltype(Header)(Data.GetPointer(sizeof(*Header)));
+	}* Header = decltype(Header)(Reader.GetPointer(sizeof(*Header)));
 	if (Header == nullptr)
 	{
 		return false;
@@ -644,7 +645,7 @@ bool FAnalysisEngine::EstablishTransport(FStreamReader::FData& Data)
 	// used to validate a inbound socket connection and then discarded.
 	if (Header->TransportVersion == 'E' || Header->TransportVersion == 'T')
 	{
-		const uint32* Magic = (const uint32*)(Data.GetPointer(sizeof(*Magic)));
+		const uint32* Magic = (const uint32*)(Reader.GetPointer(sizeof(*Magic)));
 		if (*Magic == 'ECRT')
 		{
 			// Source is big-endian which we don't currently support
@@ -653,8 +654,8 @@ bool FAnalysisEngine::EstablishTransport(FStreamReader::FData& Data)
 
 		if (*Magic == 'TRCE')
 		{
-			Data.Advance(sizeof(*Magic));
-			return EstablishTransport(Data);
+			Reader.Advance(sizeof(*Magic));
+			return EstablishTransport(Reader);
 		}
 
 		return false;
@@ -684,28 +685,28 @@ bool FAnalysisEngine::EstablishTransport(FStreamReader::FData& Data)
 		return false;
 	}
 
-	Data.Advance(sizeof(*Header));
+	Reader.Advance(sizeof(*Header));
 	return true;
 }
 
 ////////////////////////////////////////////////////////////////////////////////
-bool FAnalysisEngine::OnData(FStreamReader::FData& Data)
+bool FAnalysisEngine::OnData(FStreamReader& Reader)
 {
 	if (Transport == nullptr)
 	{
 		// Ensure we've a reasonable amount of data to establish the transport with
-		if (Data.GetPointer(32) == nullptr)
+		if (Reader.GetPointer(32) == nullptr)
 		{
 			return true;
 		}
 
-		if (!EstablishTransport(Data))
+		if (!EstablishTransport(Reader))
 		{
 			return false;
 		}
 	}
 
-	Transport->SetSource(Data);
+	Transport->SetReader(Reader);
 	bool bRet = (this->*ProtocolHandler)();
 
 	// If there's no analyzers left we might as well not continue
diff --git a/Engine/Source/Developer/TraceAnalysis/Private/Analysis/Engine.h b/Engine/Source/Developer/TraceAnalysis/Private/Analysis/Engine.h
--- a/Engine/Source/Developer/TraceAnalysis/Private/Analysis/Engine.h
+++ b/Engine/Source/Developer/TraceAnalysis/Private/Analysis/Engine.h
@@ -3,13 +3,14 @@
 #pragma once
 
 #include "Containers/Array.h"
-#include "DataStream.h"
 #include "Trace/Analysis.h"
 #include "Trace/Analyzer.h"
 
 namespace Trace
 {
 
+class FStreamReader;
+
 ////////////////////////////////////////////////////////////////////////////////
 class FAnalysisEngine
 	: public IAnalyzer
@@ -19,7 +20,7 @@ public:
 	struct				FEventDataInfo;
 						FAnalysisEngine(TArray<IAnalyzer*>&& InAnalyzers);
 						~FAnalysisEngine();
-	bool				OnData(FStreamReader::FData& Data);
+	bool				OnData(FStreamReader& Reader);
 
 private:
 	typedef bool (FAnalysisEngine::*ProtocolHandlerType)();
@@ -40,7 +41,7 @@ private:
 	void				OnNewEventInternal(const FOnEventContext& Context);
 	void				OnNewEventProtocol0(FDispatchBuilder& Builder, const void* EventData);
 
-	bool				EstablishTransport(FStreamReader::FData& Data);
+	bool				EstablishTransport(FStreamReader& Reader);
 	bool				OnDataProtocol0();
 	bool				AddDispatch(FDispatch* Dispatch);
 	template <typename ImplType>
diff --git a/Engine/Source/Developer/TraceAnalysis/Private/Analysis/Processor.cpp b/Engine/Source/Developer/TraceAnalysis/Private/Analysis/Processor.cpp
--- a/Engine/Source/Developer/TraceAnalysis/Private/Analysis/Processor.cpp
+++ b/Engine/Source/Developer/TraceAnalysis/Private/Analysis/Processor.cpp
@@ -4,6 +4,7 @@
 #include "HAL/Event.h"
 #include "HAL/PlatformProcess.h"
 #include "HAL/RunnableThread.h"
+#include "StreamReader.h"
 #include "Templates/UnrealTemplate.h"
 #include "Trace/Analysis.h"
 
@@ -32,16 +33,18 @@ FAnalysisProcessor::FImpl::~FImpl()
 ////////////////////////////////////////////////////////////////////////////////
 uint32 FAnalysisProcessor::FImpl::Run()
 {
-	for (FStreamReader Reader(DataStream); FStreamReader::FData* Data = Reader.Read();)
+	FStreamReaderDetail Reader(DataStream);
+
+	while (!StopEvent->Wait(0, true))
 	{
-		if (StopEvent->Wait(0, true))
+		UnpausedEvent->Wait();
+
+		if (!Reader.Read())
 		{
 			break;
 		}
 
-		UnpausedEvent->Wait();
-
-		if (!AnalysisEngine.OnData(*Data))
+		if (!AnalysisEngine.OnData(Reader))
 		{
 			break;
 		}
diff --git a/Engine/Source/Developer/TraceAnalysis/Private/Analysis/StreamReader.cpp b/Engine/Source/Developer/TraceAnalysis/Private/Analysis/StreamReader.cpp
new file mode 100644
--- /dev/null
+++ b/Engine/Source/Developer/TraceAnalysis/Private/Analysis/StreamReader.cpp
@@ -0,0 +1,63 @@
+// Copyright 1998-2019 Epic Games, Inc. All Rights Reserved.
+
+#include "StreamReader.h"
+#include "DataStream.h"
+
+namespace Trace
+{
+
+////////////////////////////////////////////////////////////////////////////////
+const uint8* FStreamReader::GetPointer(uint32 Size) const
+{
+	if (Cursor + Size > End)
+	{
+		return nullptr;
+	}
+
+	return Cursor;
+}
+
+////////////////////////////////////////////////////////////////////////////////
+void FStreamReader::Advance(uint32 Size)
+{
+	Cursor += Size;
+	check(Cursor <= End);
+}
+
+
+
+////////////////////////////////////////////////////////////////////////////////
+FStreamReaderDetail::FStreamReaderDetail(IInDataStream& InDataStream)
+: DataStream(InDataStream)
+{
+	Buffer = new uint8[BufferSize];
+}
+
+////////////////////////////////////////////////////////////////////////////////
+FStreamReaderDetail::~FStreamReaderDetail()
+{
+	delete[] Buffer;
+}
+
+////////////////////////////////////////////////////////////////////////////////
+bool FStreamReaderDetail::Read()
+{
+	int32 Remaining = int32(UPTRINT(End - Cursor));
+	if (Remaining != 0)
+	{
+		memcpy(Buffer, Cursor, Remaining);
+	}
+
+	uint8* Dest = Buffer + Remaining;
+	int32 ReadSize = DataStream.Read(Dest, BufferSize - Remaining);
+	if (ReadSize <= 0)
+	{
+		return false;
+	}
+
+	Cursor = Buffer;
+	End = Dest + ReadSize;
+	return true;
+}
+
+} // namespace Trace
diff --git a/Engine/Source/Developer/TraceAnalysis/Private/Analysis/StreamReader.h b/Engine/Source/Developer/TraceAnalysis/Private/Analysis/StreamReader.h
new file mode 100644
--- /dev/null
+++ b/Engine/Source/Developer/TraceAnalysis/Private/Analysis/StreamReader.h
@@ -0,0 +1,37 @@
+// Copyright 1998-2019 Epic Games, Inc. All Rights Reserved.
+
+#pragma once
+
+namespace Trace
+{
+
+class IInDataStream;
+
+////////////////////////////////////////////////////////////////////////////////
+class FStreamReader
+{
+public:
+	const uint8*		GetPointer(uint32 Size) const;
+	void				Advance(uint32 Size);
+
+protected:
+	const uint8*		Cursor = nullptr;
+	const uint8*		End = nullptr;
+};
+
+////////////////////////////////////////////////////////////////////////////////
+class FStreamReaderDetail
+	: public FStreamReader
+{
+public:
+						FStreamReaderDetail(IInDataStream& InDataStream);
+						~FStreamReaderDetail();
+	bool				Read();
+
+private:
+	static const uint32	BufferSize = 1 << 18;
+	IInDataStream&		DataStream;
+	uint8*				Buffer;
+};
+
+} // namespace Trace
diff --git a/Engine/Source/Developer/TraceAnalysis/Private/Analysis/Transport/Transport.h b/Engine/Source/Developer/TraceAnalysis/Private/Analysis/Transport/Transport.h
--- a/Engine/Source/Developer/TraceAnalysis/Private/Analysis/Transport/Transport.h
+++ b/Engine/Source/Developer/TraceAnalysis/Private/Analysis/Transport/Transport.h
@@ -2,7 +2,7 @@
 
 #pragma once
 
-#include "DataStream.h"
+#include "Analysis/StreamReader.h"
 
 namespace Trace
 {
@@ -12,7 +12,7 @@ class FTransport
 {
 public:
 	virtual					~FTransport() {}
-	void					SetSource(FStreamReader::FData& InSource);
+	void					SetReader(FStreamReader& InReader);
 	template <typename RetType>
 	RetType const*			GetPointer();
 	template <typename RetType>
@@ -21,13 +21,13 @@ public:
 
 protected:
 	virtual const uint8*	GetPointerImpl(uint32 BlockSize);
-	FStreamReader::FData*	Source;
+	FStreamReader*			Reader;
 };
 
 ////////////////////////////////////////////////////////////////////////////////
-inline void FTransport::SetSource(FStreamReader::FData& InSource)
+inline void FTransport::SetReader(FStreamReader& InReader)
 {
-	Source = &InSource;
+	Reader = &InReader;
 }
 
 ////////////////////////////////////////////////////////////////////////////////
@@ -47,13 +47,13 @@ inline RetType const* FTransport::GetPointer(uint32 BlockSize)
 ////////////////////////////////////////////////////////////////////////////////
 inline void FTransport::Advance(uint32 BlockSize)
 {
-	Source->Advance(BlockSize);
+	Reader->Advance(BlockSize);
 }
 
 ////////////////////////////////////////////////////////////////////////////////
 inline const uint8* FTransport::GetPointerImpl(uint32 BlockSize)
 {
-	return Source->GetPointer(BlockSize);
+	return Reader->GetPointer(BlockSize);
 }
 
 } // namespace Trace
diff --git a/Engine/Source/Developer/TraceAnalysis/Private/DataStream.cpp b/Engine/Source/Developer/TraceAnalysis/Private/DataStream.cpp
--- a/Engine/Source/Developer/TraceAnalysis/Private/DataStream.cpp
+++ b/Engine/Source/Developer/TraceAnalysis/Private/DataStream.cpp
@@ -61,17 +61,6 @@ void FFileStream::OpenFileInternal()
 	}
 }
 
-FStreamReader::FStreamReader(IInDataStream& InDataStream)
-	: DataStream(InDataStream)
-{
-	Buffer = new uint8[BufferSize];
-}
-
-FStreamReader::~FStreamReader()
-{
-	delete[] Buffer;
-}
-
 IInDataStream* DataStream_ReadFile(const TCHAR* FilePath)
 {
 	IPlatformFile& FileSystem = IPlatformFile::GetPlatformPhysical();
diff --git a/Engine/Source/Developer/TraceAnalysis/Private/DataStream.h b/Engine/Source/Developer/TraceAnalysis/Private/DataStream.h
--- a/Engine/Source/Developer/TraceAnalysis/Private/DataStream.h
+++ b/Engine/Source/Developer/TraceAnalysis/Private/DataStream.h
@@ -31,69 +31,4 @@ private:
 	uint64			End = 0;
 };
 
-////////////////////////////////////////////////////////////////////////////////
-class FStreamReader
-{
-public:
-	class FData
-	{
-		friend			FStreamReader;
-		const uint8*	Cursor = nullptr;
-		const uint8*	End = nullptr;
-
-	public:
-		const uint8*	GetPointer(uint32 Size) const;
-		void			Advance(uint32 Size);
-	};
-
-						FStreamReader(IInDataStream& InDataStream);
-						~FStreamReader();
-	FData*				Read();
-
-private:
-	static const uint32	BufferSize = 1 << 18;
-	FData				Data;
-	IInDataStream&		DataStream;
-	uint8*				Buffer;
-};
-
-////////////////////////////////////////////////////////////////////////////////
-inline FStreamReader::FData* FStreamReader::Read()
-{
-	int32 Remaining = int32(UPTRINT(Data.End - Data.Cursor));
-	if (Remaining != 0)
-	{
-		memcpy(Buffer, Data.Cursor, Remaining);
-	}
-
-	uint8* Dest = Buffer + Remaining;
-	int32 ReadSize = DataStream.Read(Dest, BufferSize - Remaining);
-	if (ReadSize <= 0)
-	{
-		return nullptr;
-	}
-
-	Data.Cursor = Buffer;
-	Data.End = Dest + ReadSize;
-	return &Data;
-}
-
-////////////////////////////////////////////////////////////////////////////////
-inline const uint8* FStreamReader::FData::GetPointer(uint32 Size) const
-{
-	if (Cursor + Size > End)
-	{
-		return nullptr;
-	}
-
-	return Cursor;
-}
-
-////////////////////////////////////////////////////////////////////////////////
-inline void FStreamReader::FData::Advance(uint32 Size)
-{
-	Cursor += Size;
-	check(Cursor <= End);
-}
-
 } // namespace Trace
