Fixed non-unity compile errors.

#rnx

diff --git a/Engine/Source/Developer/TraceAnalysis/Private/Analysis/Transport/PacketTransport.cpp b/Engine/Source/Developer/TraceAnalysis/Private/Analysis/Transport/PacketTransport.cpp
--- a/Engine/Source/Developer/TraceAnalysis/Private/Analysis/Transport/PacketTransport.cpp
+++ b/Engine/Source/Developer/TraceAnalysis/Private/Analysis/Transport/PacketTransport.cpp
@@ -1,6 +1,9 @@
 // Copyright 1998-2019 Epic Games, Inc. All Rights Reserved.
 
 #include "PacketTransport.h"
+#include "HAL/UnrealMemory.h"
+
+#include <initializer_list>
 
 namespace Trace
 {
@@ -82,7 +85,7 @@ FPacketTransport::FPacketNode* FPacketTransport::AllocateNode()
 	}
 	else
 	{
-		Node = (FPacketNode*)GMalloc->Malloc(sizeof(FPacketNode) + MaxPacketSize);
+		Node = (FPacketNode*)FMemory::Malloc(sizeof(FPacketNode) + MaxPacketSize);
 	}
 
 	Node->Cursor = 0;
