Use FMemory to allocate the stream reader's buffer instead of new[]/delete[]

#rnx

diff --git a/Engine/Source/Developer/TraceAnalysis/Private/Analysis/StreamReader.cpp b/Engine/Source/Developer/TraceAnalysis/Private/Analysis/StreamReader.cpp
--- a/Engine/Source/Developer/TraceAnalysis/Private/Analysis/StreamReader.cpp
+++ b/Engine/Source/Developer/TraceAnalysis/Private/Analysis/StreamReader.cpp
@@ -2,6 +2,7 @@
 
 #include "StreamReader.h"
 #include "DataStream.h"
+#include "HAL/UnrealMemory.h"
 
 namespace Trace
 {
@@ -30,13 +31,13 @@ void FStreamReader::Advance(uint32 Size)
 FStreamReaderDetail::FStreamReaderDetail(IInDataStream& InDataStream)
 : DataStream(InDataStream)
 {
-	Buffer = new uint8[BufferSize];
+	Buffer = (uint8*)FMemory::Malloc(BufferSize);
 }
 
 ////////////////////////////////////////////////////////////////////////////////
 FStreamReaderDetail::~FStreamReaderDetail()
 {
-	delete[] Buffer;
+	FMemory::Free(Buffer);
 }
 
 ////////////////////////////////////////////////////////////////////////////////
