We shouldn't be using a memcpy() here as the source and destination may overlap

#rnx

diff --git a/Engine/Source/Developer/TraceAnalysis/Private/Analysis/StreamReader.cpp b/Engine/Source/Developer/TraceAnalysis/Private/Analysis/StreamReader.cpp
--- a/Engine/Source/Developer/TraceAnalysis/Private/Analysis/StreamReader.cpp
+++ b/Engine/Source/Developer/TraceAnalysis/Private/Analysis/StreamReader.cpp
@@ -46,7 +46,7 @@ bool FStreamReaderDetail::Read()
 	int32 Remaining = int32(UPTRINT(End - Cursor));
 	if (Remaining != 0)
 	{
-		memcpy(Buffer, Cursor, Remaining);
+		memmove(Buffer, Cursor, Remaining);
 	}
 
 	uint8* Dest = Buffer + Remaining;
