Don't use a hard-coded buffer size so it can scale to demand

#rnx

diff --git a/Engine/Source/Developer/TraceAnalysis/Private/Analysis/StreamReader.cpp b/Engine/Source/Developer/TraceAnalysis/Private/Analysis/StreamReader.cpp
--- a/Engine/Source/Developer/TraceAnalysis/Private/Analysis/StreamReader.cpp
+++ b/Engine/Source/Developer/TraceAnalysis/Private/Analysis/StreamReader.cpp
@@ -12,6 +12,7 @@ const uint8* FStreamReader::GetPointer(uint32 Size) const
 {
 	if (Cursor + Size > End)
 	{
+		LastRequestedSize = Size;
 		return nullptr;
 	}
 
@@ -30,6 +31,7 @@ void FStreamReader::Advance(uint32 Size)
 ////////////////////////////////////////////////////////////////////////////////
 FStreamReaderDetail::FStreamReaderDetail(IInDataStream& InDataStream)
 : DataStream(InDataStream)
+, BufferSize(1024)
 {
 	Buffer = (uint8*)FMemory::Malloc(BufferSize);
 }
@@ -44,6 +46,20 @@ FStreamReaderDetail::~FStreamReaderDetail()
 bool FStreamReaderDetail::Read()
 {
 	int32 Remaining = int32(UPTRINT(End - Cursor));
+
+	uint32 Demand = Remaining + LastRequestedSize;
+	if (Demand > BufferSize)
+	{
+		const uint32 GrowthSizeMask = (8 << 10) - 1;
+		BufferSize = (Demand + GrowthSizeMask) & ~GrowthSizeMask;
+		uint8* NextBuffer = (uint8*)FMemory::Realloc(Buffer, BufferSize);
+
+		SIZE_T Delta = NextBuffer - Buffer;
+		Cursor += Delta;
+		End += Delta;
+		Buffer += Delta;
+	}
+
 	if (Remaining != 0)
 	{
 		memmove(Buffer, Cursor, Remaining);
diff --git a/Engine/Source/Developer/TraceAnalysis/Private/Analysis/StreamReader.h b/Engine/Source/Developer/TraceAnalysis/Private/Analysis/StreamReader.h
--- a/Engine/Source/Developer/TraceAnalysis/Private/Analysis/StreamReader.h
+++ b/Engine/Source/Developer/TraceAnalysis/Private/Analysis/StreamReader.h
@@ -19,6 +19,7 @@ public:
 protected:
 	const uint8*		Cursor = nullptr;
 	const uint8*		End = nullptr;
+	mutable uint32		LastRequestedSize = 0;
 };
 
 ////////////////////////////////////////////////////////////////////////////////
@@ -31,9 +32,9 @@ public:
 	bool				Read();
 
 private:
-	static const uint32	BufferSize = 1 << 18;
 	IInDataStream&		DataStream;
 	uint8*				Buffer;
+	uint32				BufferSize;
 };
 
 } // namespace Trace
