Store the protocol version instead of comparing function pointers

#rnx

diff --git a/Engine/Source/Developer/TraceAnalysis/Private/Analysis/Engine.cpp b/Engine/Source/Developer/TraceAnalysis/Private/Analysis/Engine.cpp
--- a/Engine/Source/Developer/TraceAnalysis/Private/Analysis/Engine.cpp
+++ b/Engine/Source/Developer/TraceAnalysis/Private/Analysis/Engine.cpp
@@ -540,9 +540,9 @@ void FAnalysisEngine::OnNewEventInternal(const FOnEventContext& Context)
 	const FEventDataInfo& EventData = (const FEventDataInfo&)(Context.EventData);
 
 	FDispatchBuilder Builder;
-	if (ProtocolHandler == &FAnalysisEngine::OnDataProtocol0)
+	switch (ProtocolVersion)
 	{
-		OnNewEventProtocol0(Builder, EventData.Ptr);
+	case int(Protocol0::EProtocol::Id): OnNewEventProtocol0(Builder, EventData.Ptr); break;
 	}
 
 	// Get the dispatch and add it into the dispatch table. Fail gently if there
@@ -670,7 +670,8 @@ bool FAnalysisEngine::EstablishTransport(FStreamReader& Reader)
 	//case 'T':	/* See the magic above */ break;
 	}
 
-	switch (Header->ProtocolVersion)
+	ProtocolVersion = Header->ProtocolVersion;
+	switch (ProtocolVersion)
 	{
 	case int(Protocol0::EProtocol::Id):
 		ProtocolHandler = &FAnalysisEngine::OnDataProtocol0;
diff --git a/Engine/Source/Developer/TraceAnalysis/Private/Analysis/Engine.h b/Engine/Source/Developer/TraceAnalysis/Private/Analysis/Engine.h
--- a/Engine/Source/Developer/TraceAnalysis/Private/Analysis/Engine.h
+++ b/Engine/Source/Developer/TraceAnalysis/Private/Analysis/Engine.h
@@ -55,6 +55,7 @@ private:
 	TArray<FDispatch*>	Dispatches;
 	class FTransport*	Transport = nullptr;
 	ProtocolHandlerType	ProtocolHandler = nullptr;
+	uint8				ProtocolVersion = 0;
 };
 
 } // namespace Trace
