Having these as scoped enums really doesn't work very well when their purposes to give some symantic meaning to bytes in a stream received across the wire.

#rnx

diff --git a/Engine/Source/Developer/TraceAnalysis/Private/Analysis/Engine.cpp b/Engine/Source/Developer/TraceAnalysis/Private/Analysis/Engine.cpp
--- a/Engine/Source/Developer/TraceAnalysis/Private/Analysis/Engine.cpp
+++ b/Engine/Source/Developer/TraceAnalysis/Private/Analysis/Engine.cpp
@@ -542,7 +542,7 @@ void FAnalysisEngine::OnNewEventInternal(const FOnEventContext& Context)
 	FDispatchBuilder Builder;
 	switch (ProtocolVersion)
 	{
-	case int(Protocol0::EProtocol::Id): OnNewEventProtocol0(Builder, EventData.Ptr); break;
+	case Protocol0::EProtocol::Id: OnNewEventProtocol0(Builder, EventData.Ptr); break;
 	}
 
 	// Get the dispatch and add it into the dispatch table. Fail gently if there
@@ -663,8 +663,8 @@ bool FAnalysisEngine::EstablishTransport(FStreamReader& Reader)
 
 	switch (Header->TransportVersion)
 	{
-	case int(ETransport::Raw):		Transport = new FTransport(); break;
-	case int(ETransport::Packet):	Transport = new FPacketTransport(); break;
+	case ETransport::Raw:		Transport = new FTransport(); break;
+	case ETransport::Packet:	Transport = new FPacketTransport(); break;
 	default:					return false;
 	//case 'E':	/* See the magic above */ break;
 	//case 'T':	/* See the magic above */ break;
@@ -673,7 +673,7 @@ bool FAnalysisEngine::EstablishTransport(FStreamReader& Reader)
 	ProtocolVersion = Header->ProtocolVersion;
 	switch (ProtocolVersion)
 	{
-	case int(Protocol0::EProtocol::Id):
+	case Protocol0::EProtocol::Id:
 		ProtocolHandler = &FAnalysisEngine::OnDataProtocol0;
 		{
 			FDispatchBuilder Builder;
diff --git a/Engine/Source/Runtime/TraceLog/Public/Trace/Detail/Protocol.h b/Engine/Source/Runtime/TraceLog/Public/Trace/Detail/Protocol.h
--- a/Engine/Source/Runtime/TraceLog/Public/Trace/Detail/Protocol.h
+++ b/Engine/Source/Runtime/TraceLog/Public/Trace/Detail/Protocol.h
@@ -24,7 +24,7 @@
 namespace Trace
 {
 
-enum class ETransport : uint8
+enum ETransport : uint8
 {
 	_Unused		= 0,
 	Raw			= 1,
diff --git a/Engine/Source/Runtime/TraceLog/Public/Trace/Detail/Protocols/Protocol0.h b/Engine/Source/Runtime/TraceLog/Public/Trace/Detail/Protocols/Protocol0.h
--- a/Engine/Source/Runtime/TraceLog/Public/Trace/Detail/Protocols/Protocol0.h
+++ b/Engine/Source/Runtime/TraceLog/Public/Trace/Detail/Protocols/Protocol0.h
@@ -12,7 +12,7 @@ namespace Protocol0
 {
 
 ////////////////////////////////////////////////////////////////////////////////
-enum class EProtocol : uint8 { Id = 0 };
+enum EProtocol : uint8 { Id = 0 };
 
 ////////////////////////////////////////////////////////////////////////////////
 enum : uint8
