Refactored the classes used to buffer an incoming trace stream and present it in a consumable manner to the anaylsis framework. Now the buffer is filled up from the outside so becomes agnostic of source.

#rnx

diff --git a/Engine/Source/Developer/TraceAnalysis/Private/Analysis/Processor.cpp b/Engine/Source/Developer/TraceAnalysis/Private/Analysis/Processor.cpp
--- a/Engine/Source/Developer/TraceAnalysis/Private/Analysis/Processor.cpp
+++ b/Engine/Source/Developer/TraceAnalysis/Private/Analysis/Processor.cpp
@@ -1,6 +1,7 @@
 // Copyright 1998-2019 Epic Games, Inc. All Rights Reserved.
 
 #include "Processor.h"
+#include "DataStream.h"
 #include "HAL/Event.h"
 #include "HAL/PlatformProcess.h"
 #include "HAL/RunnableThread.h"
@@ -33,18 +34,23 @@ FAnalysisProcessor::FImpl::~FImpl()
 ////////////////////////////////////////////////////////////////////////////////
 uint32 FAnalysisProcessor::FImpl::Run()
 {
-	FStreamReaderDetail Reader(DataStream);
+	FStreamBuffer Buffer;
 
 	while (!StopEvent->Wait(0, true))
 	{
 		UnpausedEvent->Wait();
 
-		if (!Reader.Read())
+		int32 BytesRead = Buffer.Fill([&] (uint8* Out, uint32 Size)
+		{
+			return DataStream.Read(Out, Size);
+		});
+
+		if (BytesRead <= 0)
 		{
 			break;
 		}
 
-		if (!AnalysisEngine.OnData(Reader))
+		if (!AnalysisEngine.OnData(Buffer))
 		{
 			break;
 		}
diff --git a/Engine/Source/Developer/TraceAnalysis/Private/Analysis/StreamReader.cpp b/Engine/Source/Developer/TraceAnalysis/Private/Analysis/StreamReader.cpp
--- a/Engine/Source/Developer/TraceAnalysis/Private/Analysis/StreamReader.cpp
+++ b/Engine/Source/Developer/TraceAnalysis/Private/Analysis/StreamReader.cpp
@@ -1,80 +1,71 @@
 // Copyright 1998-2019 Epic Games, Inc. All Rights Reserved.
 
 #include "StreamReader.h"
-#include "DataStream.h"
 #include "HAL/UnrealMemory.h"
+#include "Math/UnrealMath.h"
 
 namespace Trace
 {
 
 ////////////////////////////////////////////////////////////////////////////////
-const uint8* FStreamReader::GetPointer(uint32 Size) const
+FStreamReader::~FStreamReader()
+{
+	FMemory::Free(Buffer);
+}
+
+////////////////////////////////////////////////////////////////////////////////
+const uint8* FStreamReader::GetPointer(uint32 Size)
 {
 	if (Cursor + Size > End)
 	{
-		LastRequestedSize = Size;
+		DemandHint = FMath::Max(DemandHint, Size);
 		return nullptr;
 	}
 
-	return Cursor;
+	return Buffer + Cursor;
 }
 
 ////////////////////////////////////////////////////////////////////////////////
 void FStreamReader::Advance(uint32 Size)
 {
 	Cursor += Size;
-	check(Cursor <= End);
 }
 
-
-
 ////////////////////////////////////////////////////////////////////////////////
-FStreamReaderDetail::FStreamReaderDetail(IInDataStream& InDataStream)
-: DataStream(InDataStream)
-, BufferSize(1024)
+bool FStreamReader::IsEmpty() const
 {
-	Buffer = (uint8*)FMemory::Malloc(BufferSize);
+	return Cursor >= End;
 }
 
-////////////////////////////////////////////////////////////////////////////////
-FStreamReaderDetail::~FStreamReaderDetail()
-{
-	FMemory::Free(Buffer);
-}
+
 
 ////////////////////////////////////////////////////////////////////////////////
-bool FStreamReaderDetail::Read()
+void FStreamBuffer::Consolidate()
 {
-	int32 Remaining = int32(UPTRINT(End - Cursor));
+	int32 Remaining = End - Cursor;
+	DemandHint += Remaining;
 
-	uint32 Demand = Remaining + LastRequestedSize;
-	if (Demand > BufferSize)
+	if (DemandHint >= BufferSize)
 	{
 		const uint32 GrowthSizeMask = (8 << 10) - 1;
-		BufferSize = (Demand + GrowthSizeMask) & ~GrowthSizeMask;
-		uint8* NextBuffer = (uint8*)FMemory::Realloc(Buffer, BufferSize);
-
-		SIZE_T Delta = NextBuffer - Buffer;
-		Cursor += Delta;
-		End += Delta;
-		Buffer += Delta;
+		BufferSize = (DemandHint + GrowthSizeMask + 1) & ~GrowthSizeMask;
+		Buffer = (uint8*)FMemory::Realloc(Buffer, BufferSize);
 	}
 
-	if (Remaining != 0)
+	if (!Remaining)
 	{
-		memmove(Buffer, Cursor, Remaining);
+		Cursor = 0;
+		End = 0;
 	}
-
-	uint8* Dest = Buffer + Remaining;
-	int32 ReadSize = DataStream.Read(Dest, BufferSize - Remaining);
-	if (ReadSize <= 0)
+	else if (Cursor)
 	{
-		return false;
+		memmove(Buffer, Buffer + Cursor, Remaining);
+		Cursor = 0;
+		End = Remaining;
+		check(End <= BufferSize);
 	}
 
-	Cursor = Buffer;
-	End = Dest + ReadSize;
-	return true;
+	DemandHint = 0;
 }
 
 } // namespace Trace
diff --git a/Engine/Source/Developer/TraceAnalysis/Private/Analysis/StreamReader.h b/Engine/Source/Developer/TraceAnalysis/Private/Analysis/StreamReader.h
--- a/Engine/Source/Developer/TraceAnalysis/Private/Analysis/StreamReader.h
+++ b/Engine/Source/Developer/TraceAnalysis/Private/Analysis/StreamReader.h
@@ -7,34 +7,60 @@
 namespace Trace
 {
 
-class IInDataStream;
-
 ////////////////////////////////////////////////////////////////////////////////
 class FStreamReader
 {
 public:
-	const uint8*		GetPointer(uint32 Size) const;
-	void				Advance(uint32 Size);
+								~FStreamReader();
+	template <typename Type>
+	Type const*					GetPointer();
+	const uint8*				GetPointer(uint32 Size);
+	void						Advance(uint32 Size);
+	bool						IsEmpty() const;
 
 protected:
-	const uint8*		Cursor = nullptr;
-	const uint8*		End = nullptr;
-	mutable uint32		LastRequestedSize = 0;
+	uint8*						Buffer = nullptr;
+	uint32						DemandHint = 0;
+	uint32						Cursor = 0;
+	uint32						End = 0;
 };
 
 ////////////////////////////////////////////////////////////////////////////////
-class FStreamReaderDetail
+template <typename Type>
+Type const* FStreamReader::GetPointer()
+{
+	return (Type const*)GetPointer(sizeof(Type));
+}
+
+
+
+////////////////////////////////////////////////////////////////////////////////
+class FStreamBuffer
 	: public FStreamReader
 {
 public:
-						FStreamReaderDetail(IInDataStream& InDataStream);
-						~FStreamReaderDetail();
-	bool				Read();
-
-private:
-	IInDataStream&		DataStream;
-	uint8*				Buffer;
-	uint32				BufferSize;
+	template <typename Lambda>
+	int32						Fill(Lambda&& Source);
+
+protected:
+	void						Consolidate();
+	uint32						BufferSize = 0;
 };
 
+////////////////////////////////////////////////////////////////////////////////
+template <typename Lambda>
+inline int32 FStreamBuffer::Fill(Lambda&& Source)
+{
+	Consolidate();
+
+	uint8* Dest = Buffer + End;
+	int32 ReadSize = Source(Dest, BufferSize - End);
+	if (ReadSize > 0)
+	{
+		End += ReadSize;
+	}
+
+	return ReadSize;
+}
+
 } // namespace Trace
