Added two methods for appending data to a stream buffer.

#rnx

diff --git a/Engine/Source/Developer/TraceAnalysis/Private/Analysis/StreamReader.cpp b/Engine/Source/Developer/TraceAnalysis/Private/Analysis/StreamReader.cpp
--- a/Engine/Source/Developer/TraceAnalysis/Private/Analysis/StreamReader.cpp
+++ b/Engine/Source/Developer/TraceAnalysis/Private/Analysis/StreamReader.cpp
@@ -39,6 +39,24 @@ bool FStreamReader::IsEmpty() const
 
 
 
+////////////////////////////////////////////////////////////////////////////////
+void FStreamBuffer::Append(const uint8* Data, uint32 Size)
+{
+	uint8* Out = Append(Size);
+	FMemory::Memcpy(Out, Data, Size);
+}
+
+////////////////////////////////////////////////////////////////////////////////
+uint8* FStreamBuffer::Append(uint32 Size)
+{
+	DemandHint = FMath::Max(Size, DemandHint);
+	Consolidate();
+	uint8* Out = Buffer + End;
+	End += Size;
+	check(End <= BufferSize);
+	return Out;
+}
+
 ////////////////////////////////////////////////////////////////////////////////
 void FStreamBuffer::Consolidate()
 {
diff --git a/Engine/Source/Developer/TraceAnalysis/Private/Analysis/StreamReader.h b/Engine/Source/Developer/TraceAnalysis/Private/Analysis/StreamReader.h
--- a/Engine/Source/Developer/TraceAnalysis/Private/Analysis/StreamReader.h
+++ b/Engine/Source/Developer/TraceAnalysis/Private/Analysis/StreamReader.h
@@ -41,6 +41,8 @@ class FStreamBuffer
 public:
 	template <typename Lambda>
 	int32						Fill(Lambda&& Source);
+	void						Append(const uint8* Data, uint32 Size);
+	uint8*						Append(uint32 Size);
 
 protected:
 	void						Consolidate();
