Resynchronised traced events across threads by using a rolling serial number. The coordination of the serialised order happens at analysis time.

#rnx

diff --git a/Engine/Source/Developer/TraceAnalysis/Private/Analysis/Engine.cpp b/Engine/Source/Developer/TraceAnalysis/Private/Analysis/Engine.cpp
--- a/Engine/Source/Developer/TraceAnalysis/Private/Analysis/Engine.cpp
+++ b/Engine/Source/Developer/TraceAnalysis/Private/Analysis/Engine.cpp
@@ -10,6 +10,7 @@
 #include "Trace/Detail/Protocol.h"
 #include "Transport/PacketTransport.h"
 #include "Transport/Transport.h"
+#include "Transport/TidPacketTransport.h"
 
 namespace Trace
 {
@@ -475,6 +476,7 @@ void FAnalysisEngine::OnNewTrace(const FOnEventContext& Context)
 		}
 	}
 
+	FixedRouteCount = 0; // Disabled for now until AddRoute([ExplicitHash]) has been removed
 	TArrayView<FRoute> RouteSubset(Routes.GetData() + FixedRouteCount, Routes.Num() - FixedRouteCount);
 	Algo::SortBy(RouteSubset, [] (const FRoute& Route) { return Route.Hash; });
 
@@ -543,6 +545,7 @@ void FAnalysisEngine::OnNewEventInternal(const FOnEventContext& Context)
 	switch (ProtocolVersion)
 	{
 	case Protocol0::EProtocol::Id: OnNewEventProtocol0(Builder, EventData.Ptr); break;
+	case Protocol1::EProtocol::Id: OnNewEventProtocol1(Builder, EventData.Ptr); break;
 	}
 
 	// Get the dispatch and add it into the dispatch table. Fail gently if there
@@ -629,6 +632,12 @@ void FAnalysisEngine::OnNewEventProtocol0(FDispatchBuilder& Builder, const void*
 	}
 }
 
+////////////////////////////////////////////////////////////////////////////////
+void FAnalysisEngine::OnNewEventProtocol1(FDispatchBuilder& Builder, const void* EventData)
+{
+	return OnNewEventProtocol0(Builder, EventData);
+}
+
 ////////////////////////////////////////////////////////////////////////////////
 bool FAnalysisEngine::EstablishTransport(FStreamReader& Reader)
 {
@@ -665,6 +674,7 @@ bool FAnalysisEngine::EstablishTransport(FStreamReader& Reader)
 	{
 	case ETransport::Raw:		Transport = new FTransport(); break;
 	case ETransport::Packet:	Transport = new FPacketTransport(); break;
+	case ETransport::TidPacket:	Transport = new FTidPacketTransport(); break;
 	default:					return false;
 	//case 'E':	/* See the magic above */ break;
 	//case 'T':	/* See the magic above */ break;
@@ -682,6 +692,15 @@ bool FAnalysisEngine::EstablishTransport(FStreamReader& Reader)
 		}
 		break;
 
+	case Protocol1::EProtocol::Id:
+		ProtocolHandler = &FAnalysisEngine::OnDataProtocol1;
+		{
+			FDispatchBuilder Builder;
+			Builder.SetUid(uint16(Protocol0::FNewEventEvent::Uid));
+			AddDispatch(Builder.Finalize());
+		}
+		break;
+
 	default:
 		return false;
 	}
@@ -753,7 +772,86 @@ bool FAnalysisEngine::OnDataProtocol0()
 			return false;
 		}
 
+		FEventDataInfo EventDataInfo = { *Dispatch, Header->EventData, Header->Size };
+		const FEventData& EventData = (FEventData&)EventDataInfo;
+
+		ForEachRoute(Dispatch, [&] (IAnalyzer* Analyzer, uint16 RouteId)
+		{
+			if (!Analyzer->OnEvent(RouteId, { SessionContext, EventData }))
+			{
+				RetireAnalyzer(Analyzer);
+			}
+		});
+
 		Transport->Advance(BlockSize);
+	}
+
+	return true;
+}
+
+////////////////////////////////////////////////////////////////////////////////
+bool FAnalysisEngine::OnDataProtocol1()
+{
+	auto* InnerTransport = (FTidPacketTransport*)Transport;
+	InnerTransport->Update();
+
+	int32 EventCount;
+	do
+	{
+		EventCount = 0;
+		FTidPacketTransport::ThreadIter Iter = InnerTransport->ReadThreads();
+		while (FStreamReader* Reader = InnerTransport->GetNextThread(Iter))
+		{
+			int32 ThreadEventCount = OnDataProtocol1(*Reader);
+			if (ThreadEventCount < 0)
+			{
+				return false;
+			}
+
+			EventCount += ThreadEventCount;
+		}
+	}
+	while (EventCount);
+
+	return true;
+}
+
+////////////////////////////////////////////////////////////////////////////////
+int32 FAnalysisEngine::OnDataProtocol1(FStreamReader& Reader)
+{
+	int32 EventCount = 0;
+	while (!Reader.IsEmpty())
+	{
+		const auto* Header = Reader.GetPointer<Protocol1::FEventHeader>();
+		if (Header == nullptr)
+		{
+			break;
+		}
+
+		// Make sure we consume events in the correct order
+		if (Header->Serial != uint16(NextLogSerial))
+		{
+			break;
+		}
+		++NextLogSerial;
+
+		uint32 BlockSize = Header->Size + sizeof(Protocol1::FEventHeader);
+		if (Reader.GetPointer(BlockSize) == nullptr)
+		{
+			break;
+		}
+
+		uint16 Uid = uint16(Header->Uid & 0x3fff);
+		if (Uid >= Dispatches.Num())
+		{
+			return -1;
+		}
+
+		const FDispatch* Dispatch = Dispatches[Uid];
+		if (Dispatch == nullptr)
+		{
+			return -1;
+		}
 
 		FEventDataInfo EventDataInfo = { *Dispatch, Header->EventData, Header->Size };
 		const FEventData& EventData = (FEventData&)EventDataInfo;
@@ -765,9 +863,12 @@ bool FAnalysisEngine::OnDataProtocol0()
 				RetireAnalyzer(Analyzer);
 			}
 		});
+
+		Reader.Advance(BlockSize);
+		++EventCount;
 	}
 
-	return true;
+	return EventCount;
 }
 
 } // namespace Trace
diff --git a/Engine/Source/Developer/TraceAnalysis/Private/Analysis/Engine.h b/Engine/Source/Developer/TraceAnalysis/Private/Analysis/Engine.h
--- a/Engine/Source/Developer/TraceAnalysis/Private/Analysis/Engine.h
+++ b/Engine/Source/Developer/TraceAnalysis/Private/Analysis/Engine.h
@@ -40,9 +40,12 @@ private:
 	void				OnTiming(const FOnEventContext& Context);
 	void				OnNewEventInternal(const FOnEventContext& Context);
 	void				OnNewEventProtocol0(FDispatchBuilder& Builder, const void* EventData);
+	void				OnNewEventProtocol1(FDispatchBuilder& Builder, const void* EventData);
 
 	bool				EstablishTransport(FStreamReader& Reader);
 	bool				OnDataProtocol0();
+	bool				OnDataProtocol1();
+	int32				OnDataProtocol1(FStreamReader& Reader);
 	bool				AddDispatch(FDispatch* Dispatch);
 	template <typename ImplType>
 	void				ForEachRoute(const FDispatch* Dispatch, ImplType&& Impl);
@@ -55,6 +58,7 @@ private:
 	TArray<FDispatch*>	Dispatches;
 	class FTransport*	Transport = nullptr;
 	ProtocolHandlerType	ProtocolHandler = nullptr;
+	uint16				NextLogSerial = 0;
 	uint8				ProtocolVersion = 0;
 };
 
diff --git a/Engine/Source/Developer/TraceAnalysis/Private/Analysis/Transport/TidPacketTransport.cpp b/Engine/Source/Developer/TraceAnalysis/Private/Analysis/Transport/TidPacketTransport.cpp
new file mode 100644
--- /dev/null
+++ b/Engine/Source/Developer/TraceAnalysis/Private/Analysis/Transport/TidPacketTransport.cpp
@@ -0,0 +1,105 @@
+// Copyright 1998-2019 Epic Games, Inc. All Rights Reserved.
+
+#include "TidPacketTransport.h"
+#include "HAL/UnrealMemory.h"
+
+namespace Trace
+{
+
+////////////////////////////////////////////////////////////////////////////////
+namespace Private
+{
+
+TRACELOG_API int32 Decode(const void*, int32, void*, int32);
+
+} // namespace Private
+
+
+
+////////////////////////////////////////////////////////////////////////////////
+bool FTidPacketTransport::ReadPacket()
+{
+	struct FPacketBase
+	{
+		uint16	PacketSize;
+		uint16	ThreadId;
+	};
+
+	const auto* PacketBase = GetPointer<FPacketBase>();
+	if (PacketBase == nullptr)
+	{
+		return false;
+	}
+
+	if (GetPointer<uint8>(PacketBase->PacketSize) == nullptr)
+	{
+		return false;
+	}
+
+	FTransport::Advance(PacketBase->PacketSize);
+
+	uint32 ThreadId = PacketBase->ThreadId & ~0x8000;
+	FThreadStream& Thread = FindOrAddThread(ThreadId);
+
+	uint32 DataSize = PacketBase->PacketSize - sizeof(FPacketBase);
+	if (PacketBase->ThreadId != ThreadId)
+	{
+		uint16* DecodedSize = (uint16*)(PacketBase + 1);
+		uint8* Dest = Thread.Buffer.Append(*DecodedSize);
+		DataSize -= sizeof(*DecodedSize);
+		Private::Decode(DecodedSize + 1, DataSize, Dest, *DecodedSize);
+	}
+	else
+	{
+		Thread.Buffer.Append((uint8*)(PacketBase + 1), DataSize);
+	}
+
+	return true;
+}
+
+////////////////////////////////////////////////////////////////////////////////
+FTidPacketTransport::FThreadStream& FTidPacketTransport::FindOrAddThread(uint32 ThreadId)
+{
+	for (auto& Thread : Threads)
+	{
+		if (Thread.ThreadId == ThreadId)
+		{
+			return Thread;
+		}
+	}
+
+	FThreadStream Thread;
+	Thread.ThreadId = ThreadId;
+	Threads.Add(Thread);
+	return Threads.Last();
+}
+
+////////////////////////////////////////////////////////////////////////////////
+void FTidPacketTransport::Update()
+{
+	while (ReadPacket());
+	Threads.RemoveAllSwap([] (const FThreadStream& Thread) { return Thread.Buffer.IsEmpty(); });
+}
+
+////////////////////////////////////////////////////////////////////////////////
+FTidPacketTransport::ThreadIter FTidPacketTransport::ReadThreads()
+{
+	return ThreadIter(0);
+}
+
+////////////////////////////////////////////////////////////////////////////////
+FStreamReader* FTidPacketTransport::GetNextThread(ThreadIter& Iter)
+{
+	while (Iter < Threads.Num())
+	{
+		FThreadStream& Thread = Threads[Iter++];
+		if (!Thread.Buffer.IsEmpty())
+		{
+			return &(Thread.Buffer);
+		}
+	}
+
+	return nullptr;
+}
+
+} // namespace Trace
diff --git a/Engine/Source/Developer/TraceAnalysis/Private/Analysis/Transport/TidPacketTransport.h b/Engine/Source/Developer/TraceAnalysis/Private/Analysis/Transport/TidPacketTransport.h
new file mode 100644
--- /dev/null
+++ b/Engine/Source/Developer/TraceAnalysis/Private/Analysis/Transport/TidPacketTransport.h
@@ -0,0 +1,37 @@
+// Copyright 1998-2019 Epic Games, Inc. All Rights Reserved.
+
+#pragma once
+
+#include "Transport.h"
+#include "Analysis/StreamReader.h"
+
+namespace Trace
+{
+
+class FStreamReader;
+
+////////////////////////////////////////////////////////////////////////////////
+class FTidPacketTransport
+	: public FTransport
+{
+public:
+	typedef UPTRINT ThreadIter;
+
+	void					Update();
+	ThreadIter				ReadThreads();
+	FStreamReader*			GetNextThread(ThreadIter& Iter);
+
+private:
+	struct FThreadStream
+	{
+		FStreamBuffer		Buffer;
+		uint32				ThreadId;
+	};
+
+	bool					ReadPacket();
+	FThreadStream&			FindOrAddThread(uint32 ThreadId);
+	static const uint32		MaxPacketSize = 8192;
+	TArray<FThreadStream>	Threads;
+};
+
+} // namespace Trace
diff --git a/Engine/Source/Runtime/TraceLog/Private/Trace/Writer.cpp b/Engine/Source/Runtime/TraceLog/Private/Trace/Writer.cpp
--- a/Engine/Source/Runtime/TraceLog/Private/Trace/Writer.cpp
+++ b/Engine/Source/Runtime/TraceLog/Private/Trace/Writer.cpp
@@ -79,9 +79,10 @@ void Writer_InitializeTiming()
 
 
 ////////////////////////////////////////////////////////////////////////////////
-thread_local		FWriteTlsContext TlsContext;
-uint8				FWriteTlsContext::DefaultBuffer[];	// = {}
-UPTRINT volatile	FWriteTlsContext::ThreadIdCounter;	// = 0;
+thread_local					FWriteTlsContext TlsContext;
+uint8							FWriteTlsContext::DefaultBuffer[];	// = {}
+UPTRINT volatile				FWriteTlsContext::ThreadIdCounter;	// = 0;
+TRACELOG_API uint32 volatile	GLogSerial;							// = 0;
 
 ////////////////////////////////////////////////////////////////////////////////
 FWriteTlsContext::FWriteTlsContext()
@@ -458,14 +459,14 @@ static void Writer_ConsumeEvents()
 		}
 	};
 
-	auto Send = [&] (uint8* __restrict Data, uint32 Size)
+	auto Send = [&] (FWriteBuffer* Buffer, uint8* __restrict Data, uint32 Size)
 	{
 		BytesSent += Size;
 
 		struct FPacketBase
 		{
-			uint16 Serial;
 			uint16 PacketSize;
+			uint16 ThreadId;
 		};
 
 		// Smaller buffers usually aren't redundant enough to benefit from being
@@ -485,7 +486,7 @@ static void Writer_ConsumeEvents()
 			};
 
 			FPacket Packet;
-			Packet.Serial = 0x8000;
+			Packet.ThreadId = 0x8000 | uint16(Buffer->ThreadId & 0x7fff);
 			Packet.DecodedSize = uint16(Size);
 			Packet.PacketSize = Encode(Data, Packet.DecodedSize, Packet.Data, sizeof(Packet.Data));
 			Packet.PacketSize += sizeof(FPacketEncoded);
@@ -498,8 +499,9 @@ static void Writer_ConsumeEvents()
 			Data -= sizeof(FPacketBase);
 			Size += sizeof(FPacketBase);
 			auto* Packet = (FPacketBase*)Data;
-			Packet->Serial = 0;
+			Packet->ThreadId = uint16(Buffer->ThreadId & 0x7fff);
 			Packet->PacketSize = uint16(Size);
+
 			SendInner(Data, Size);
 		}
 	};
@@ -561,7 +563,7 @@ static void Writer_ConsumeEvents()
 
 		if (uint32 SizeToReap = uint32(Committed - Buffer->Reaped))
 		{
-			Send(Buffer->Reaped, SizeToReap);
+			Send(Buffer, Buffer->Reaped, SizeToReap);
 			Buffer->Reaped = Committed;
 		}
 	}
@@ -577,7 +579,7 @@ static void Writer_ConsumeEvents()
 
 		if (uint32 SizeToReap = uint32(Buffer->Cursor - Buffer->Reaped))
 		{
-			Send(Buffer->Reaped, SizeToReap);
+			Send(Buffer, Buffer->Reaped, SizeToReap);
 		}
 	}
 
@@ -637,8 +639,8 @@ static void Writer_UpdateData()
 
 		// Stream header
 		const struct {
-			uint8 TransportVersion	= uint8(ETransport::Packet);
-			uint8 ProtocolVersion	= uint8(EProtocol::Id);
+			uint8 TransportVersion	= ETransport::TidPacket;
+			uint8 ProtocolVersion	= EProtocol::Id;
 		} TransportHeader;
 		bOk &= IoWrite(GDataHandle, &TransportHeader, sizeof(TransportHeader));
 
diff --git a/Engine/Source/Runtime/TraceLog/Public/Trace/Detail/Protocol.h b/Engine/Source/Runtime/TraceLog/Public/Trace/Detail/Protocol.h
--- a/Engine/Source/Runtime/TraceLog/Public/Trace/Detail/Protocol.h
+++ b/Engine/Source/Runtime/TraceLog/Public/Trace/Detail/Protocol.h
@@ -6,7 +6,7 @@
 
 ////////////////////////////////////////////////////////////////////////////////
 #if UE_TRACE_ENABLED
-#	define TRACE_PRIVATE_PROTOCOL_0
+#	define TRACE_PRIVATE_PROTOCOL_1
 #endif
 
 #if defined(_MSC_VER)
@@ -15,6 +15,7 @@
 #endif
 
 #include "Protocols/Protocol0.h"
+#include "Protocols/Protocol1.h"
 
 #if defined(_MSC_VER)
 	#pragma warning(pop)
@@ -29,6 +30,7 @@ enum ETransport : uint8
 	_Unused		= 0,
 	Raw			= 1,
 	Packet		= 2,
+	TidPacket	= 3,
 };
 
 } // namespace Trace
diff --git a/Engine/Source/Runtime/TraceLog/Public/Trace/Detail/Protocols/Protocol1.h b/Engine/Source/Runtime/TraceLog/Public/Trace/Detail/Protocols/Protocol1.h
new file mode 100644
--- /dev/null
+++ b/Engine/Source/Runtime/TraceLog/Public/Trace/Detail/Protocols/Protocol1.h
@@ -0,0 +1,31 @@
+// Copyright 1998-2019 Epic Games, Inc. All Rights Reserved.
+
+#pragma once
+
+namespace Trace
+{
+
+#if defined(TRACE_PRIVATE_PROTOCOL_1)
+inline
+#endif
+namespace Protocol1
+{
+
+////////////////////////////////////////////////////////////////////////////////
+enum EProtocol : uint8 { Id = 1 };
+
+////////////////////////////////////////////////////////////////////////////////
+using Protocol0::EFieldType;
+using Protocol0::FNewEventEvent;
+
+////////////////////////////////////////////////////////////////////////////////
+struct FEventHeader
+{
+	uint16		Uid;
+	uint16		Size;
+	uint16		Serial;
+	uint8		EventData[];
+};
+
+} // namespace Protocol1
+} // namespace Trace
diff --git a/Engine/Source/Runtime/TraceLog/Public/Trace/Detail/Writer.inl b/Engine/Source/Runtime/TraceLog/Public/Trace/Detail/Writer.inl
--- a/Engine/Source/Runtime/TraceLog/Public/Trace/Detail/Writer.inl
+++ b/Engine/Source/Runtime/TraceLog/Public/Trace/Detail/Writer.inl
@@ -45,8 +45,9 @@ private:
 
 
 ////////////////////////////////////////////////////////////////////////////////
-TRACELOG_API FWriteBuffer* Writer_NextBuffer(uint16);
-TRACELOG_API FWriteBuffer* Writer_GetBuffer();
+extern TRACELOG_API uint32 volatile	GLogSerial;
+TRACELOG_API FWriteBuffer*			Writer_NextBuffer(uint16);
+TRACELOG_API FWriteBuffer*			Writer_GetBuffer();
 
 ////////////////////////////////////////////////////////////////////////////////
 #if IS_MONOLITHIC
@@ -73,21 +74,22 @@ inline FLogInstance Writer_BeginLog(uint16 EventUid, uint16 Size)
 {
 	using namespace Private;
 
-	static const uint32 HeaderSize = sizeof(FEventHeader);
-	uint32 AllocSize = Size + HeaderSize;
-
 	FWriteBuffer* Buffer = Writer_GetBuffer();
-	uint8* Cursor = (Buffer->Cursor += AllocSize);
-	if (UNLIKELY(UPTRINT(Cursor) > UPTRINT(Buffer)))
+	uint32 AllocSize = Size + sizeof(FEventHeader);
+	Buffer->Cursor += AllocSize;
+	if (UNLIKELY(UPTRINT(Buffer->Cursor) > UPTRINT(Buffer)))
 	{
 		Buffer = Writer_NextBuffer(AllocSize);
-		Cursor = Buffer->Cursor;
 	}
-	Cursor -= AllocSize;
 
-	uint32* Out = (uint32*)Cursor;
-	Out[0] = (uint32(Size) << 16)|uint32(EventUid);
-	return {(uint8*)(Out + 1), Buffer};
+	uint8* Cursor = Buffer->Cursor - Size;
+
+	auto* Header = (uint16*)(Cursor); // FEventHeader
+	Header[-1] = uint16(AtomicIncrementRelaxed(&GLogSerial));
+	Header[-2] = Size;
+	Header[-3] = EventUid;
+
+	return {Cursor, Buffer};
 }
 
 ////////////////////////////////////////////////////////////////////////////////
