Deleted unused member.

#rnx

diff --git a/Engine/Source/Developer/TraceAnalysis/Private/Analysis/Transport/TidPacketTransport.h b/Engine/Source/Developer/TraceAnalysis/Private/Analysis/Transport/TidPacketTransport.h
--- a/Engine/Source/Developer/TraceAnalysis/Private/Analysis/Transport/TidPacketTransport.h
+++ b/Engine/Source/Developer/TraceAnalysis/Private/Analysis/Transport/TidPacketTransport.h
@@ -30,7 +30,6 @@ private:
 
 	bool					ReadPacket();
 	FThreadStream&			FindOrAddThread(uint32 ThreadId);
-	static const uint32		MaxPacketSize = 8192;
 	TArray<FThreadStream>	Threads;
 };
 
