Allow stream readers to back up if any subsequent dependent reads can't be completed yet.

#rnx

diff --git a/Engine/Source/Developer/TraceAnalysis/Private/Analysis/StreamReader.cpp b/Engine/Source/Developer/TraceAnalysis/Private/Analysis/StreamReader.cpp
--- a/Engine/Source/Developer/TraceAnalysis/Private/Analysis/StreamReader.cpp
+++ b/Engine/Source/Developer/TraceAnalysis/Private/Analysis/StreamReader.cpp
@@ -37,6 +37,22 @@ bool FStreamReader::IsEmpty() const
 	return Cursor >= End;
 }
 
+////////////////////////////////////////////////////////////////////////////////
+struct FMark* FStreamReader::SaveMark() const
+{
+	return (FMark*)(UPTRINT(Cursor));
+}
+
+////////////////////////////////////////////////////////////////////////////////
+void FStreamReader::RestoreMark(struct FMark* Mark)
+{
+	Cursor = uint32(UPTRINT(Mark));
+	if (Cursor > End)
+	{
+		Cursor = End;
+	}
+}
+
 
 
 ////////////////////////////////////////////////////////////////////////////////
diff --git a/Engine/Source/Developer/TraceAnalysis/Private/Analysis/StreamReader.h b/Engine/Source/Developer/TraceAnalysis/Private/Analysis/StreamReader.h
--- a/Engine/Source/Developer/TraceAnalysis/Private/Analysis/StreamReader.h
+++ b/Engine/Source/Developer/TraceAnalysis/Private/Analysis/StreamReader.h
@@ -17,6 +17,8 @@ public:
 	const uint8*				GetPointer(uint32 Size);
 	void						Advance(uint32 Size);
 	bool						IsEmpty() const;
+	struct FMark*				SaveMark() const;
+	void						RestoreMark(struct FMark* Mark);
 
 protected:
 	uint8*						Buffer = nullptr;
