Track the end-of-file condition in stream readers so that those reading that stream can know if they can expect more data or not.

#rnx

diff --git a/Engine/Source/Developer/TraceAnalysis/Private/Analysis/StreamReader.cpp b/Engine/Source/Developer/TraceAnalysis/Private/Analysis/StreamReader.cpp
--- a/Engine/Source/Developer/TraceAnalysis/Private/Analysis/StreamReader.cpp
+++ b/Engine/Source/Developer/TraceAnalysis/Private/Analysis/StreamReader.cpp
@@ -37,6 +37,12 @@ bool FStreamReader::IsEmpty() const
 	return Cursor >= End;
 }
 
+////////////////////////////////////////////////////////////////////////////////
+bool FStreamReader::IsEof() const
+{
+	return bEof;
+}
+
 ////////////////////////////////////////////////////////////////////////////////
 struct FMark* FStreamReader::SaveMark() const
 {
@@ -102,4 +108,10 @@ void FStreamBuffer::Consolidate()
 	DemandHint = 0;
 }
 
+////////////////////////////////////////////////////////////////////////////////
+void FStreamBuffer::SetEof()
+{
+	bEof = true;
+}
+
 } // namespace Trace
diff --git a/Engine/Source/Developer/TraceAnalysis/Private/Analysis/StreamReader.h b/Engine/Source/Developer/TraceAnalysis/Private/Analysis/StreamReader.h
--- a/Engine/Source/Developer/TraceAnalysis/Private/Analysis/StreamReader.h
+++ b/Engine/Source/Developer/TraceAnalysis/Private/Analysis/StreamReader.h
@@ -17,6 +17,7 @@ public:
 	const uint8*				GetPointer(uint32 Size);
 	void						Advance(uint32 Size);
 	bool						IsEmpty() const;
+	bool						IsEof() const;
 	struct FMark*				SaveMark() const;
 	void						RestoreMark(struct FMark* Mark);
 
@@ -25,6 +26,7 @@ protected:
 	uint32						DemandHint = 0;
 	uint32						Cursor = 0;
 	uint32						End = 0;
+	bool						bEof = false;
 };
 
 ////////////////////////////////////////////////////////////////////////////////
@@ -45,6 +47,7 @@ public:
 	int32						Fill(Lambda&& Source);
 	void						Append(const uint8* Data, uint32 Size);
 	uint8*						Append(uint32 Size);
+	void						SetEof();
 
 protected:
 	void						Consolidate();
@@ -63,6 +66,10 @@ inline int32 FStreamBuffer::Fill(Lambda&& Source)
 	{
 		End += ReadSize;
 	}
+	else
+	{
+		SetEof();
+	}
 
 	return ReadSize;
 }
