Allow one more pass through the analysis before stopping due to an end of fil (unless an error occured). This ensures any sub-buffering gets to see the EOF condition.

#rnx

diff --git a/Engine/Source/Developer/TraceAnalysis/Private/Analysis/Processor.cpp b/Engine/Source/Developer/TraceAnalysis/Private/Analysis/Processor.cpp
--- a/Engine/Source/Developer/TraceAnalysis/Private/Analysis/Processor.cpp
+++ b/Engine/Source/Developer/TraceAnalysis/Private/Analysis/Processor.cpp
@@ -45,7 +45,7 @@ uint32 FAnalysisProcessor::FImpl::Run()
 			return DataStream.Read(Out, Size);
 		});
 
-		if (BytesRead <= 0)
+		if (BytesRead < 0)
 		{
 			break;
 		}
@@ -54,6 +54,11 @@ uint32 FAnalysisProcessor::FImpl::Run()
 		{
 			break;
 		}
+
+		if (Buffer.IsEof())
+		{
+			break;
+		}
 	}
 
 	return 0;
