Pass through EOF state to each threads' buffer.

#rnx
#ushell-unshelve of 9957048

diff --git a/Engine/Source/Developer/TraceAnalysis/Private/Analysis/Transport/TidPacketTransport.cpp b/Engine/Source/Developer/TraceAnalysis/Private/Analysis/Transport/TidPacketTransport.cpp
--- a/Engine/Source/Developer/TraceAnalysis/Private/Analysis/Transport/TidPacketTransport.cpp
+++ b/Engine/Source/Developer/TraceAnalysis/Private/Analysis/Transport/TidPacketTransport.cpp
@@ -79,6 +79,14 @@ void FTidPacketTransport::Update()
 {
 	while (ReadPacket());
 	Threads.RemoveAllSwap([] (const FThreadStream& Thread) { return Thread.Buffer.IsEmpty(); });
+
+	if (bool bEof = Reader->IsEof())
+	{
+		for (auto& Thread : Threads)
+		{
+			Thread.Buffer.SetEof();
+		}
+	}
 }
 
 ////////////////////////////////////////////////////////////////////////////////
