Let's not sort fields. We're not leveraging the sorted order and we will soon be relying on field indices so it'd be convenient to keep the original order.

#rnx
#ushell-unshelve of 9957048

diff --git a/Engine/Source/Developer/TraceAnalysis/Private/Analysis/Engine.cpp b/Engine/Source/Developer/TraceAnalysis/Private/Analysis/Engine.cpp
--- a/Engine/Source/Developer/TraceAnalysis/Private/Analysis/Engine.cpp
+++ b/Engine/Source/Developer/TraceAnalysis/Private/Analysis/Engine.cpp
@@ -92,10 +92,6 @@ FAnalysisEngine::FDispatch* FAnalysisEngine::FDispatchBuilder::Finalize()
 	memcpy(Dispatch, Buffer.GetData(), Buffer.Num());
 	memcpy(Dispatch->Fields + Dispatch->FieldCount, NameBuffer.GetData(), NameBuffer.Num());
 
-	// Sort by hash so we can binary search when looking up.
-	TArrayView<FDispatch::FField> Fields(Dispatch->Fields, Dispatch->FieldCount);
-	Algo::SortBy(Fields, [] (const auto& Field) { return Field.Hash; });
-
 	// Fix up name offsets
 	for (int i = 0, n = Dispatch->FieldCount; i < n; ++i)
 	{
