Being able to test if a stream's EOF doesn't make sense. It results in a situation where streams are stalled because they're not sure if there's more data or not. Instead the protocol should include a terminal (or count).

#rnx
#ushell-unshelve of 9957048

diff --git a/Engine/Source/Developer/TraceAnalysis/Private/Analysis/Processor.cpp b/Engine/Source/Developer/TraceAnalysis/Private/Analysis/Processor.cpp
--- a/Engine/Source/Developer/TraceAnalysis/Private/Analysis/Processor.cpp
+++ b/Engine/Source/Developer/TraceAnalysis/Private/Analysis/Processor.cpp
@@ -45,7 +45,7 @@ uint32 FAnalysisProcessor::FImpl::Run()
 			return DataStream.Read(Out, Size);
 		});
 
-		if (BytesRead < 0)
+		if (BytesRead <= 0)
 		{
 			break;
 		}
@@ -54,11 +54,6 @@ uint32 FAnalysisProcessor::FImpl::Run()
 		{
 			break;
 		}
-
-		if (Buffer.IsEof())
-		{
-			break;
-		}
 	}
 
 	return 0;
diff --git a/Engine/Source/Developer/TraceAnalysis/Private/Analysis/StreamReader.cpp b/Engine/Source/Developer/TraceAnalysis/Private/Analysis/StreamReader.cpp
--- a/Engine/Source/Developer/TraceAnalysis/Private/Analysis/StreamReader.cpp
+++ b/Engine/Source/Developer/TraceAnalysis/Private/Analysis/StreamReader.cpp
@@ -37,12 +37,6 @@ bool FStreamReader::IsEmpty() const
 	return Cursor >= End;
 }
 
-////////////////////////////////////////////////////////////////////////////////
-bool FStreamReader::IsEof() const
-{
-	return bEof;
-}
-
 ////////////////////////////////////////////////////////////////////////////////
 struct FMark* FStreamReader::SaveMark() const
 {
@@ -108,10 +102,4 @@ void FStreamBuffer::Consolidate()
 	DemandHint = 0;
 }
 
-////////////////////////////////////////////////////////////////////////////////
-void FStreamBuffer::SetEof()
-{
-	bEof = true;
-}
-
 } // namespace Trace
diff --git a/Engine/Source/Developer/TraceAnalysis/Private/Analysis/StreamReader.h b/Engine/Source/Developer/TraceAnalysis/Private/Analysis/StreamReader.h
--- a/Engine/Source/Developer/TraceAnalysis/Private/Analysis/StreamReader.h
+++ b/Engine/Source/Developer/TraceAnalysis/Private/Analysis/StreamReader.h
@@ -17,7 +17,6 @@ public:
 	const uint8*				GetPointer(uint32 Size);
 	void						Advance(uint32 Size);
 	bool						IsEmpty() const;
-	bool						IsEof() const;
 	struct FMark*				SaveMark() const;
 	void						RestoreMark(struct FMark* Mark);
 
@@ -26,7 +25,6 @@ protected:
 	uint32						DemandHint = 0;
 	uint32						Cursor = 0;
 	uint32						End = 0;
-	bool						bEof = false;
 };
 
 ////////////////////////////////////////////////////////////////////////////////
@@ -47,7 +45,6 @@ public:
 	int32						Fill(Lambda&& Source);
 	void						Append(const uint8* Data, uint32 Size);
 	uint8*						Append(uint32 Size);
-	void						SetEof();
 
 protected:
 	void						Consolidate();
@@ -66,10 +63,6 @@ inline int32 FStreamBuffer::Fill(Lambda&& Source)
 	{
 		End += ReadSize;
 	}
-	else
-	{
-		SetEof();
-	}
 
 	return ReadSize;
 }
diff --git a/Engine/Source/Developer/TraceAnalysis/Private/Analysis/Transport/TidPacketTransport.cpp b/Engine/Source/Developer/TraceAnalysis/Private/Analysis/Transport/TidPacketTransport.cpp
--- a/Engine/Source/Developer/TraceAnalysis/Private/Analysis/Transport/TidPacketTransport.cpp
+++ b/Engine/Source/Developer/TraceAnalysis/Private/Analysis/Transport/TidPacketTransport.cpp
@@ -79,14 +79,6 @@ void FTidPacketTransport::Update()
 {
 	while (ReadPacket());
 	Threads.RemoveAllSwap([] (const FThreadStream& Thread) { return Thread.Buffer.IsEmpty(); });
-
-	if (bool bEof = Reader->IsEof())
-	{
-		for (auto& Thread : Threads)
-		{
-			Thread.Buffer.SetEof();
-		}
-	}
 }
 
 ////////////////////////////////////////////////////////////////////////////////
