Encapsulated the look up of a field's index by name into a method of FDispatch.

#rnx
#ushell-unshelve of 9957048

diff --git a/Engine/Source/Developer/TraceAnalysis/Private/Analysis/Engine.cpp b/Engine/Source/Developer/TraceAnalysis/Private/Analysis/Engine.cpp
--- a/Engine/Source/Developer/TraceAnalysis/Private/Analysis/Engine.cpp
+++ b/Engine/Source/Developer/TraceAnalysis/Private/Analysis/Engine.cpp
@@ -46,6 +46,7 @@ struct FAnalysisEngine::FDispatch
 		int16		SizeAndType;		// value == byte_size, sign == float < 0 < int
 	};
 
+	uint32			GetFieldIndex(const ANSICHAR* Name) const;
 	uint32			Hash				= 0;
 	uint16			Uid					= 0;
 	uint16			FirstRoute			= ~uint16(0);
@@ -56,6 +57,23 @@ struct FAnalysisEngine::FDispatch
 	FField			Fields[];
 };
 
+////////////////////////////////////////////////////////////////////////////////
+uint32 FAnalysisEngine::FDispatch::GetFieldIndex(const ANSICHAR* Name) const
+{
+	FFnv1aHash NameHash;
+	NameHash.Add(Name);
+
+	for (int i = 0, n = FieldCount; i < n; ++i)
+	{
+		if (Fields[i].Hash == NameHash.Get())
+		{
+			return i;
+		}
+	}
+
+	return FieldCount;
+}
+
 
 
 ////////////////////////////////////////////////////////////////////////////////
@@ -283,22 +301,17 @@ const IAnalyzer::FEventTypeInfo& IAnalyzer::FEventData::GetTypeInfo() const
 const void* IAnalyzer::FEventData::GetValueImpl(const ANSICHAR* FieldName, int16& SizeAndType) const
 {
 	const auto* Info = (const FAnalysisEngine::FEventDataInfo*)this;
+	const auto& Dispatch = Info->Dispatch;
 
-	FFnv1aHash Hash;
-	Hash.Add(FieldName);
-	uint32 NameHash = Hash.Get();
-
-	for (int i = 0, n = Info->Dispatch.FieldCount; i < n; ++i)
+	uint32 Index = Dispatch.GetFieldIndex(FieldName);
+	if (Index >= Dispatch.FieldCount)
 	{
-		const auto& Field = Info->Dispatch.Fields[i];
-		if (Field.Hash == NameHash)
-		{
-			SizeAndType = Field.SizeAndType;
-			return (Info->Ptr + Field.Offset);
-		}
+		return nullptr;
 	}
 
-	return nullptr;
+	const auto& Field = Dispatch.Fields[Index];
+	SizeAndType = Field.SizeAndType;
+	return (Info->Ptr + Field.Offset);
 }
 
 ////////////////////////////////////////////////////////////////////////////////
