Build field's properties outside of the dispatch builder so that different protocol versions can adapt to suit

#rnx
#ushell-unshelve of 9957048

diff --git a/Engine/Source/Developer/TraceAnalysis/Private/Analysis/Engine.cpp b/Engine/Source/Developer/TraceAnalysis/Private/Analysis/Engine.cpp
--- a/Engine/Source/Developer/TraceAnalysis/Private/Analysis/Engine.cpp
+++ b/Engine/Source/Developer/TraceAnalysis/Private/Analysis/Engine.cpp
@@ -80,17 +80,17 @@ uint32 FAnalysisEngine::FDispatch::GetFieldIndex(const ANSICHAR* Name) const
 class FAnalysisEngine::FDispatchBuilder
 {
 public:
-					FDispatchBuilder();
-	void			SetUid(uint16 Uid);
-	void			SetLoggerName(const ANSICHAR* Name, int32 NameSize=-1);
-	void			SetEventName(const ANSICHAR* Name, int32 NameSize=-1);
-	void			AddField(const ANSICHAR* Name, int32 NameSize, uint16 Offset, uint16 Size, FEventFieldInfo::EType TypeId, uint16 TypeSize);
-	FDispatch*		Finalize();
+						FDispatchBuilder();
+	void				SetUid(uint16 Uid);
+	void				SetLoggerName(const ANSICHAR* Name, int32 NameSize=-1);
+	void				SetEventName(const ANSICHAR* Name, int32 NameSize=-1);
+	FDispatch::FField&	AddField(const ANSICHAR* Name, int32 NameSize, uint16 Size);
+	FDispatch*			Finalize();
 
 private:
-	uint32			AppendName(const ANSICHAR* Name, int32 NameSize);
-	TArray<uint8>	Buffer;
-	TArray<uint8>	NameBuffer;
+	uint32				AppendName(const ANSICHAR* Name, int32 NameSize);
+	TArray<uint8>		Buffer;
+	TArray<uint8>		NameBuffer;
 };
 
 ////////////////////////////////////////////////////////////////////////////////
@@ -154,18 +154,12 @@ void FAnalysisEngine::FDispatchBuilder::SetEventName(const ANSICHAR* Name, int32
 }
 
 ////////////////////////////////////////////////////////////////////////////////
-void FAnalysisEngine::FDispatchBuilder::AddField(const ANSICHAR* Name, int32 NameSize, uint16 Offset, uint16 Size, FEventFieldInfo::EType TypeId, uint16 TypeSize)
+FAnalysisEngine::FDispatch::FField& FAnalysisEngine::FDispatchBuilder::AddField(const ANSICHAR* Name, int32 NameSize, uint16 Size)
 {
 	int32 Bufoff = Buffer.AddUninitialized(sizeof(FDispatch::FField));
 	auto* Field = (FDispatch::FField*)(Buffer.GetData() + Bufoff);
 	Field->NameOffset = AppendName(Name, NameSize);
-	Field->Offset = Offset;
 	Field->Size = Size;
-	Field->SizeAndType = TypeSize;
-	if (TypeId == FEventFieldInfo::EType::Float)
-	{
-		Field->SizeAndType *= -1;
-	}
 
 	FFnv1aHash Hash;
 	Hash.Add((const uint8*)Name, NameSize);
@@ -174,6 +168,8 @@ void FAnalysisEngine::FDispatchBuilder::AddField(const ANSICHAR* Name, int32 Nam
 	auto* Dispatch = (FDispatch*)(Buffer.GetData());
 	Dispatch->FieldCount++;
 	Dispatch->EventSize += Field->Size;
+
+	return *Field;
 }
 
 ////////////////////////////////////////////////////////////////////////////////
@@ -628,14 +624,14 @@ void FAnalysisEngine::OnNewEventProtocol0(FDispatchBuilder& Builder, const void*
 		const auto& Field = NewEvent.Fields[i];
 
 		uint16 TypeSize = 1 << (Field.TypeInfo & Protocol0::Field_Pow2SizeMask);
-
-		auto TypeId = FEventFieldInfo::EType::Integer;
-		if ((Field.TypeInfo & Protocol0::Field_CategoryMask) == Protocol0::Field_Float)
+		if (Field.TypeInfo & Protocol0::Field_Float)
 		{
-			TypeId = FEventFieldInfo::EType::Float;
+			TypeSize = -TypeSize;
 		}
 
-		Builder.AddField(NameCursor, Field.NameSize, Field.Offset, Field.Size, TypeId, TypeSize);
+		auto& OutField = Builder.AddField(NameCursor, Field.NameSize, Field.Size);
+		OutField.Offset = Field.Offset;
+		OutField.SizeAndType = TypeSize;
 
 		NameCursor += Field.NameSize;
 	}
