Moved the enum of known event IDs into the protocol so it can be changed while maintaining backwards compatibility.

#rnx
#ushell-unshelve of 9957048

diff --git a/Engine/Source/Developer/TraceAnalysis/Private/Analysis/Engine.cpp b/Engine/Source/Developer/TraceAnalysis/Private/Analysis/Engine.cpp
--- a/Engine/Source/Developer/TraceAnalysis/Private/Analysis/Engine.cpp
+++ b/Engine/Source/Developer/TraceAnalysis/Private/Analysis/Engine.cpp
@@ -692,7 +692,7 @@ bool FAnalysisEngine::EstablishTransport(FStreamReader& Reader)
 		ProtocolHandler = &FAnalysisEngine::OnDataProtocol0;
 		{
 			FDispatchBuilder Builder;
-			Builder.SetUid(uint16(Protocol0::FNewEventEvent::Uid));
+			Builder.SetUid(uint16(Protocol0::EKnownEventUids::NewEvent));
 			AddDispatch(Builder.Finalize());
 		}
 		break;
@@ -701,7 +701,7 @@ bool FAnalysisEngine::EstablishTransport(FStreamReader& Reader)
 		ProtocolHandler = &FAnalysisEngine::OnDataProtocol1;
 		{
 			FDispatchBuilder Builder;
-			Builder.SetUid(uint16(Protocol0::FNewEventEvent::Uid));
+			Builder.SetUid(uint16(Protocol0::EKnownEventUids::NewEvent));
 			AddDispatch(Builder.Finalize());
 		}
 		break;
@@ -765,7 +765,7 @@ bool FAnalysisEngine::OnDataProtocol0()
 			break;
 		}
 
-		uint16 Uid = uint16(Header->Uid & ((1 << 14) - 1));
+		uint16 Uid = uint16(Header->Uid) & uint16(Protocol0::EKnownEventUids::UidMask);
 		if (Uid >= Dispatches.Num())
 		{
 			return false;
@@ -845,7 +845,7 @@ int32 FAnalysisEngine::OnDataProtocol1(FStreamReader& Reader)
 			break;
 		}
 
-		uint16 Uid = uint16(Header->Uid & 0x3fff);
+		uint16 Uid = Header->Uid & uint16(Protocol1::EKnownEventUids::UidMask);
 		if (Uid >= Dispatches.Num())
 		{
 			return -1;
diff --git a/Engine/Source/Runtime/TraceLog/Private/Trace/Writer.cpp b/Engine/Source/Runtime/TraceLog/Private/Trace/Writer.cpp
--- a/Engine/Source/Runtime/TraceLog/Private/Trace/Writer.cpp
+++ b/Engine/Source/Runtime/TraceLog/Private/Trace/Writer.cpp
@@ -1093,18 +1093,6 @@ bool Writer_WriteTo(const ANSICHAR* Path)
 static UPTRINT volatile		GEventUidCounter;	// = 0;
 static FEventDef* volatile	GHeadEvent;			// = nullptr;
 
-////////////////////////////////////////////////////////////////////////////////
-enum class EKnownEventUids : uint16
-{
-	NewEvent		= FNewEventEvent::Uid,
-	User,
-	Max				= (1 << 14) - 1, // ...leaves two MSB bits for other uses.
-	UidMask			= Max,
-	Invalid			= Max,
-	Flag_Unused		= 1 << 14,
-	Flag_Important	= 1 << 15,
-};
-
 ////////////////////////////////////////////////////////////////////////////////
 template <typename ElementType>
 static uint32 Writer_EventGetHash(const ElementType* Input, int32 Length=-1)
diff --git a/Engine/Source/Runtime/TraceLog/Public/Trace/Detail/Protocols/Protocol0.h b/Engine/Source/Runtime/TraceLog/Public/Trace/Detail/Protocols/Protocol0.h
--- a/Engine/Source/Runtime/TraceLog/Public/Trace/Detail/Protocols/Protocol0.h
+++ b/Engine/Source/Runtime/TraceLog/Public/Trace/Detail/Protocols/Protocol0.h
@@ -51,7 +51,6 @@ enum class EFieldType : uint8
 ////////////////////////////////////////////////////////////////////////////////
 struct FNewEventEvent
 {
-	enum : uint16 { Uid = 0 };
 
 	uint16		EventUid;
 	uint16		FieldCount;
@@ -67,6 +66,18 @@ struct FNewEventEvent
 	/*uint8		NameData[]*/
 };
 
+////////////////////////////////////////////////////////////////////////////////
+enum class EKnownEventUids : uint16
+{
+	NewEvent,
+	User,
+	Max				= (1 << 14) - 1, // ...leaves two MSB bits for other uses.
+	UidMask			= Max,
+	Invalid			= Max,
+	Flag_Important	= 1 << 14,
+	Flag_Unused		= 1 << 15,
+};
+
 ////////////////////////////////////////////////////////////////////////////////
 struct FEventHeader
 {
diff --git a/Engine/Source/Runtime/TraceLog/Public/Trace/Detail/Protocols/Protocol1.h b/Engine/Source/Runtime/TraceLog/Public/Trace/Detail/Protocols/Protocol1.h
--- a/Engine/Source/Runtime/TraceLog/Public/Trace/Detail/Protocols/Protocol1.h
+++ b/Engine/Source/Runtime/TraceLog/Public/Trace/Detail/Protocols/Protocol1.h
@@ -17,6 +17,7 @@ enum EProtocol : uint8 { Id = 1 };
 ////////////////////////////////////////////////////////////////////////////////
 using Protocol0::EFieldType;
 using Protocol0::FNewEventEvent;
+using Protocol0::EKnownEventUids;
 
 ////////////////////////////////////////////////////////////////////////////////
 struct FEventHeader
