Send the important flag one the one-time trace event that describes and event instead of tagging each event sent

#rnx
#ushell-unshelve of 9957048

diff --git a/Engine/Source/Developer/TraceAnalysis/Private/Analysis/Engine.cpp b/Engine/Source/Developer/TraceAnalysis/Private/Analysis/Engine.cpp
--- a/Engine/Source/Developer/TraceAnalysis/Private/Analysis/Engine.cpp
+++ b/Engine/Source/Developer/TraceAnalysis/Private/Analysis/Engine.cpp
@@ -37,6 +37,11 @@ private:
 ////////////////////////////////////////////////////////////////////////////////
 struct FAnalysisEngine::FDispatch
 {
+	enum
+	{
+		Flag_Important		= 1 << 0,
+	};
+
 	struct FField
 	{
 		uint32		Hash;
@@ -50,7 +55,8 @@ struct FAnalysisEngine::FDispatch
 	uint32			Hash				= 0;
 	uint16			Uid					= 0;
 	uint16			FirstRoute			= ~uint16(0);
-	uint16			FieldCount			= 0;
+	uint8			FieldCount			= 0;
+	uint8			Flags				= 0;
 	uint16			EventSize			= 0;
 	uint16			LoggerNameOffset	= 0;	// From FDispatch ptr
 	uint16			EventNameOffset		= 0;	// From FDispatch ptr
@@ -84,6 +90,7 @@ public:
 	void				SetUid(uint16 Uid);
 	void				SetLoggerName(const ANSICHAR* Name, int32 NameSize=-1);
 	void				SetEventName(const ANSICHAR* Name, int32 NameSize=-1);
+	void				SetImportant();
 	FDispatch::FField&	AddField(const ANSICHAR* Name, int32 NameSize, uint16 Size);
 	FDispatch*			Finalize();
 
@@ -153,6 +160,13 @@ void FAnalysisEngine::FDispatchBuilder::SetEventName(const ANSICHAR* Name, int32
 	Dispatch->EventNameOffset = AppendName(Name, NameSize);
 }
 
+////////////////////////////////////////////////////////////////////////////////
+void FAnalysisEngine::FDispatchBuilder::SetImportant()
+{
+	auto* Dispatch = (FDispatch*)(Buffer.GetData());
+	Dispatch->Flags |= FDispatch::Flag_Important;
+}
+
 ////////////////////////////////////////////////////////////////////////////////
 FAnalysisEngine::FDispatch::FField& FAnalysisEngine::FDispatchBuilder::AddField(const ANSICHAR* Name, int32 NameSize, uint16 Size)
 {
@@ -640,7 +654,14 @@ void FAnalysisEngine::OnNewEventProtocol0(FDispatchBuilder& Builder, const void*
 ////////////////////////////////////////////////////////////////////////////////
 void FAnalysisEngine::OnNewEventProtocol1(FDispatchBuilder& Builder, const void* EventData)
 {
-	return OnNewEventProtocol0(Builder, EventData);
+	OnNewEventProtocol0(Builder, EventData);
+
+	const auto& NewEvent = *(Protocol1::FNewEventEvent*)(EventData);
+
+	if (NewEvent.Flags & int(Protocol1::EEventFlags::Important))
+	{
+		Builder.SetImportant();
+	}
 }
 
 ////////////////////////////////////////////////////////////////////////////////
@@ -845,7 +866,7 @@ int32 FAnalysisEngine::OnDataProtocol1(FStreamReader& Reader)
 			break;
 		}
 
-		uint16 Uid = Header->Uid & uint16(Protocol1::EKnownEventUids::UidMask);
+		uint16 Uid = Header->Uid;
 		if (Uid >= Dispatches.Num())
 		{
 			return -1;
diff --git a/Engine/Source/Runtime/TraceLog/Private/Trace/Writer.cpp b/Engine/Source/Runtime/TraceLog/Private/Trace/Writer.cpp
--- a/Engine/Source/Runtime/TraceLog/Private/Trace/Writer.cpp
+++ b/Engine/Source/Runtime/TraceLog/Private/Trace/Writer.cpp
@@ -1146,11 +1146,6 @@ void Writer_EventCreate(
 		return;
 	}
 
-	if (Flags & FEventDef::Flag_Important)
-	{
-		Uid |= uint16(EKnownEventUids::Flag_Important);
-	}
-
 	uint32 LoggerHash = Writer_EventGetHash(LoggerName.Ptr);
 	uint32 NameHash = Writer_EventGetHash(EventName.Ptr);
 
@@ -1170,7 +1165,7 @@ void Writer_EventCreate(
 	}
 
 	// Allocate the new event event in the log stream.
-	uint16 EventUid = uint16(EKnownEventUids::NewEvent)|uint16(EKnownEventUids::Flag_Important);
+	uint16 EventUid = uint16(EKnownEventUids::NewEvent);
 	uint16 EventSize = sizeof(FNewEventEvent);
 	EventSize += sizeof(FNewEventEvent::Fields[0]) * FieldCount;
 	EventSize += NamesSize;
@@ -1179,12 +1174,18 @@ void Writer_EventCreate(
 	auto& Event = *(FNewEventEvent*)(LogInstance.Ptr);
 
 	// Write event's main properties.
-	Event.EventUid = uint16(Uid) & uint16(EKnownEventUids::UidMask);
+	Event.EventUid = uint16(Uid);
 	Event.LoggerNameSize = LoggerName.Length;
 	Event.EventNameSize = EventName.Length;
+	Event.Flags = 0;
+
+	if (Flags & FEventDef::Flag_Important)
+	{
+		Event.Flags |= uint8(EEventFlags::Important);
+	}
 
 	// Write details about event's fields
-	Event.FieldCount = FieldCount;
+	Event.FieldCount = uint8(FieldCount);
 	for (uint32 i = 0; i < FieldCount; ++i)
 	{
 		const FFieldDesc& FieldDesc = FieldDescs[i];
diff --git a/Engine/Source/Runtime/TraceLog/Public/Trace/Detail/Protocols/Protocol0.h b/Engine/Source/Runtime/TraceLog/Public/Trace/Detail/Protocols/Protocol0.h
--- a/Engine/Source/Runtime/TraceLog/Public/Trace/Detail/Protocols/Protocol0.h
+++ b/Engine/Source/Runtime/TraceLog/Public/Trace/Detail/Protocols/Protocol0.h
@@ -51,9 +51,9 @@ enum class EFieldType : uint8
 ////////////////////////////////////////////////////////////////////////////////
 struct FNewEventEvent
 {
-
 	uint16		EventUid;
-	uint16		FieldCount;
+	uint8		FieldCount;
+	uint8		Flags;
 	uint8		LoggerNameSize;
 	uint8		EventNameSize;
 	struct
diff --git a/Engine/Source/Runtime/TraceLog/Public/Trace/Detail/Protocols/Protocol1.h b/Engine/Source/Runtime/TraceLog/Public/Trace/Detail/Protocols/Protocol1.h
--- a/Engine/Source/Runtime/TraceLog/Public/Trace/Detail/Protocols/Protocol1.h
+++ b/Engine/Source/Runtime/TraceLog/Public/Trace/Detail/Protocols/Protocol1.h
@@ -12,13 +12,19 @@ namespace Protocol1
 {
 
 ////////////////////////////////////////////////////////////////////////////////
-enum EProtocol : uint8 { Id = 1 };
+enum EProtocol : uint8 { Id = 3 };
 
 ////////////////////////////////////////////////////////////////////////////////
 using Protocol0::EFieldType;
 using Protocol0::FNewEventEvent;
 using Protocol0::EKnownEventUids;
 
+////////////////////////////////////////////////////////////////////////////////
+enum class EEventFlags : uint8
+{
+	Important		= 1 << 0,
+};
+
 ////////////////////////////////////////////////////////////////////////////////
 struct FEventHeader
 {
