Extended trace protocol so events can be marked as possibly expecting auxilary data. Such events are terminated with a null byte so if the following byte is non-zero it is expected to be an additional payload for the event. Nothing is currently using this.

#rnx
#ushell-unshelve of 9957048

diff --git a/Engine/Source/Developer/TraceAnalysis/Private/Analysis/Engine.cpp b/Engine/Source/Developer/TraceAnalysis/Private/Analysis/Engine.cpp
--- a/Engine/Source/Developer/TraceAnalysis/Private/Analysis/Engine.cpp
+++ b/Engine/Source/Developer/TraceAnalysis/Private/Analysis/Engine.cpp
@@ -34,12 +34,29 @@ private:
 
 
 
+////////////////////////////////////////////////////////////////////////////////
+struct FAuxData
+{
+	const uint8*	Data;
+	uint32			DataSize;
+	uint16			FieldIndex;
+};
+
+////////////////////////////////////////////////////////////////////////////////
+struct FAnalysisEngine::FAuxDataCollector
+	: public TArray<FAuxData>
+{
+};
+
+
+
 ////////////////////////////////////////////////////////////////////////////////
 struct FAnalysisEngine::FDispatch
 {
 	enum
 	{
 		Flag_Important		= 1 << 0,
+		Flag_MaybeHasAux	= 1 << 1,
 	};
 
 	struct FField
@@ -91,6 +108,7 @@ public:
 	void				SetLoggerName(const ANSICHAR* Name, int32 NameSize=-1);
 	void				SetEventName(const ANSICHAR* Name, int32 NameSize=-1);
 	void				SetImportant();
+	void				SetMaybeHasAux();
 	FDispatch::FField&	AddField(const ANSICHAR* Name, int32 NameSize, uint16 Size);
 	FDispatch*			Finalize();
 
@@ -167,6 +185,13 @@ void FAnalysisEngine::FDispatchBuilder::SetImportant()
 	Dispatch->Flags |= FDispatch::Flag_Important;
 }
 
+////////////////////////////////////////////////////////////////////////////////
+void FAnalysisEngine::FDispatchBuilder::SetMaybeHasAux()
+{
+	auto* Dispatch = (FDispatch*)(Buffer.GetData());
+	Dispatch->Flags |= FDispatch::Flag_MaybeHasAux;
+}
+
 ////////////////////////////////////////////////////////////////////////////////
 FAnalysisEngine::FDispatch::FField& FAnalysisEngine::FDispatchBuilder::AddField(const ANSICHAR* Name, int32 NameSize, uint16 Size)
 {
@@ -662,6 +687,11 @@ void FAnalysisEngine::OnNewEventProtocol1(FDispatchBuilder& Builder, const void*
 	{
 		Builder.SetImportant();
 	}
+
+	if (NewEvent.Flags & int(Protocol1::EEventFlags::MaybeHasAux))
+	{
+		Builder.SetMaybeHasAux();
+	}
 }
 
 ////////////////////////////////////////////////////////////////////////////////
@@ -848,6 +878,8 @@ int32 FAnalysisEngine::OnDataProtocol1(FStreamReader& Reader)
 	int32 EventCount = 0;
 	while (!Reader.IsEmpty())
 	{
+		auto Mark = Reader.SaveMark();
+
 		const auto* Header = Reader.GetPointer<Protocol1::FEventHeader>();
 		if (Header == nullptr)
 		{
@@ -878,6 +910,23 @@ int32 FAnalysisEngine::OnDataProtocol1(FStreamReader& Reader)
 			return -1;
 		}
 
+		Reader.Advance(BlockSize);
+
+		FAuxDataCollector AuxCollector;
+		if (Dispatch->Flags & FDispatch::Flag_MaybeHasAux)
+		{
+			int AuxStatus = OnDataProtocol1Aux(Reader, AuxCollector);
+			if (AuxStatus == 0)
+			{
+				Reader.RestoreMark(Mark);
+				break;
+			}
+			else if (AuxStatus < 0)
+			{
+				return -1;
+			}
+		}
+
 		++NextLogSerial;
 
 		FEventDataInfo EventDataInfo = { *Dispatch, Header->EventData, Header->Size };
@@ -891,11 +940,54 @@ int32 FAnalysisEngine::OnDataProtocol1(FStreamReader& Reader)
 			}
 		});
 
-		Reader.Advance(BlockSize);
 		++EventCount;
 	}
 
 	return EventCount;
 }
 
+////////////////////////////////////////////////////////////////////////////////
+int32 FAnalysisEngine::OnDataProtocol1Aux(FStreamReader& Reader, FAuxDataCollector& Collector)
+{
+	while (true)
+	{
+		const uint8* NextByte = Reader.GetPointer<uint8>();
+		if (NextByte == nullptr)
+		{
+			return 0;
+		}
+
+		// Is the following sequence a blob of auxilary data or the null
+		// terminator byte?
+		if (NextByte[0] == 0)
+		{
+			Reader.Advance(1);
+			return 1;
+		}
+
+		// Get header and the auxilary blob's size
+		const auto* Header = Reader.GetPointer<Protocol1::FAuxHeader>();
+		if (Header == nullptr)
+		{
+			return 0;
+		}
+
+		// Check it exists
+		uint32 BlockSize = (Header->Size >> 8) + sizeof(*Header);
+		if (Reader.GetPointer(BlockSize) == nullptr)
+		{
+			return 0;
+		}
+
+		// Attach to event
+		FAuxData AuxData;
+		AuxData.Data = Header->Data;
+		AuxData.DataSize = uint32(BlockSize - sizeof(*Header));
+		AuxData.FieldIndex = uint16(Header->FieldIndex & Protocol1::FAuxHeader::FieldMask);
+		Collector.Push(AuxData);
+
+		Reader.Advance(BlockSize);
+	}
+}
+
 } // namespace Trace
diff --git a/Engine/Source/Developer/TraceAnalysis/Private/Analysis/Engine.h b/Engine/Source/Developer/TraceAnalysis/Private/Analysis/Engine.h
--- a/Engine/Source/Developer/TraceAnalysis/Private/Analysis/Engine.h
+++ b/Engine/Source/Developer/TraceAnalysis/Private/Analysis/Engine.h
@@ -17,6 +17,7 @@ class FAnalysisEngine
 {
 public:
 	struct				FDispatch;
+	struct				FAuxDataCollector;
 	struct				FEventDataInfo;
 						FAnalysisEngine(TArray<IAnalyzer*>&& InAnalyzers);
 						~FAnalysisEngine();
@@ -46,6 +47,7 @@ private:
 	bool				OnDataProtocol0();
 	bool				OnDataProtocol1();
 	int32				OnDataProtocol1(FStreamReader& Reader);
+	int32				OnDataProtocol1Aux(FStreamReader& Reader, FAuxDataCollector& Collector);
 	bool				AddDispatch(FDispatch* Dispatch);
 	template <typename ImplType>
 	void				ForEachRoute(const FDispatch* Dispatch, ImplType&& Impl);
diff --git a/Engine/Source/Runtime/TraceLog/Private/Trace/Writer.cpp b/Engine/Source/Runtime/TraceLog/Private/Trace/Writer.cpp
--- a/Engine/Source/Runtime/TraceLog/Private/Trace/Writer.cpp
+++ b/Engine/Source/Runtime/TraceLog/Private/Trace/Writer.cpp
@@ -1170,7 +1170,7 @@ void Writer_EventCreate(
 	EventSize += sizeof(FNewEventEvent::Fields[0]) * FieldCount;
 	EventSize += NamesSize;
 
-	FLogInstance LogInstance = Writer_BeginLog(EventUid, EventSize);
+	FLogInstance LogInstance = Writer_BeginLog(EventUid, EventSize, false);
 	auto& Event = *(FNewEventEvent*)(LogInstance.Ptr);
 
 	// Write event's main properties.
@@ -1184,6 +1184,11 @@ void Writer_EventCreate(
 		Event.Flags |= uint8(EEventFlags::Important);
 	}
 
+	if (Flags & FEventDef::Flag_MaybeHasAux)
+	{
+		Event.Flags |= uint8(EEventFlags::MaybeHasAux);
+	}
+
 	// Write details about event's fields
 	Event.FieldCount = uint8(FieldCount);
 	for (uint32 i = 0; i < FieldCount; ++i)
diff --git a/Engine/Source/Runtime/TraceLog/Public/Trace/Detail/EventDef.h b/Engine/Source/Runtime/TraceLog/Public/Trace/Detail/EventDef.h
--- a/Engine/Source/Runtime/TraceLog/Public/Trace/Detail/EventDef.h
+++ b/Engine/Source/Runtime/TraceLog/Public/Trace/Detail/EventDef.h
@@ -20,15 +20,16 @@ class FEventDef
 public:
 	enum
 	{
-		Flag_Always		= 1 << 0,
-		Flag_Important	= 1 << 1,
+		Flag_Always			= 1 << 0,
+		Flag_Important		= 1 << 1,
+		Flag_MaybeHasAux	= 1 << 2,
 	};
 
 	class FLogScope
 	{
 	public:
-								FLogScope(uint16 EventUid, uint16 Size);
-								FLogScope(uint16 EventUid, uint16 Size, uint16 ExtraBytes);
+								FLogScope(uint16 EventUid, uint16 Size, bool bMaybeHasAux);
+								FLogScope(uint16 EventUid, uint16 Size, bool bMaybeHasAux, uint16 ExtraBytes);
 								~FLogScope();
 		FLogInstance			Instance;
 	};
diff --git a/Engine/Source/Runtime/TraceLog/Public/Trace/Detail/EventDef.inl b/Engine/Source/Runtime/TraceLog/Public/Trace/Detail/EventDef.inl
--- a/Engine/Source/Runtime/TraceLog/Public/Trace/Detail/EventDef.inl
+++ b/Engine/Source/Runtime/TraceLog/Public/Trace/Detail/EventDef.inl
@@ -10,15 +10,15 @@ namespace Trace
 {
 
 ////////////////////////////////////////////////////////////////////////////////
-inline FEventDef::FLogScope::FLogScope(uint16 EventUid, uint16 Size)
+inline FEventDef::FLogScope::FLogScope(uint16 EventUid, uint16 Size, bool bMaybeHasAux)
 {
-	Instance = Writer_BeginLog(EventUid, Size);
+	Instance = Writer_BeginLog(EventUid, Size, bMaybeHasAux);
 }
 
 ////////////////////////////////////////////////////////////////////////////////
-inline FEventDef::FLogScope::FLogScope(uint16 EventUid, uint16 Size, uint16 ExtraBytes)
+inline FEventDef::FLogScope::FLogScope(uint16 EventUid, uint16 Size, bool bMaybeHasAux, uint16 ExtraBytes)
 {
-	Instance = Writer_BeginLog(EventUid, Size + ExtraBytes);
+	Instance = Writer_BeginLog(EventUid, Size + ExtraBytes, bMaybeHasAux);
 }
 
 ////////////////////////////////////////////////////////////////////////////////
diff --git a/Engine/Source/Runtime/TraceLog/Public/Trace/Detail/Field.h b/Engine/Source/Runtime/TraceLog/Public/Trace/Detail/Field.h
--- a/Engine/Source/Runtime/TraceLog/Public/Trace/Detail/Field.h
+++ b/Engine/Source/Runtime/TraceLog/Public/Trace/Detail/Field.h
@@ -64,6 +64,17 @@ struct FFieldDesc
 	uint8			TypeInfo;
 };
 
+
+
+////////////////////////////////////////////////////////////////////////////////
+template <int InIndex, int InOffset, typename Type> struct TField;
+
+enum class EIndexPack
+{
+	FieldCountMask	= 0xff,
+	MaybeHasAux		= 0x100,
+};
+
 ////////////////////////////////////////////////////////////////////////////////
 #define TRACE_PRIVATE_FIELD(InIndex, InOffset, Type) \
 		enum \
@@ -73,7 +84,7 @@ struct FFieldDesc
 			Tid		= TFieldType<Type>::Tid, \
 			Size	= TFieldType<Type>::Size, \
 		}; \
-		static_assert(Index <= 127, "Trace events may only have up to a maximum of 127 fields"); \
+		static_assert((Index & int(EIndexPack::FieldCountMask)) <= 127, "Trace events may only have up to a maximum of 127 fields"); \
 	private: \
 		FFieldDesc FieldDesc; \
 	public: \
@@ -132,11 +143,16 @@ struct TField
 
 ////////////////////////////////////////////////////////////////////////////////
 // Used to terminate the field list and determine an event's size.
-enum EndOfFields {};
-template <int InFieldCount, int Size>
-struct TField<InFieldCount, Size, EndOfFields>
+enum EventProps {};
+template <int InFieldCount, int InSize>
+struct TField<InFieldCount, InSize, EventProps>
 {
-	enum : uint16 { FieldCount = InFieldCount, Value = Size };
+	enum : uint16
+	{
+		FieldCount	= (InFieldCount & int(EIndexPack::FieldCountMask)),
+		Size		= InSize,
+		MaybeHasAux	= !!(InFieldCount & int(EIndexPack::MaybeHasAux)),
+	};
 };
 
 ////////////////////////////////////////////////////////////////////////////////
diff --git a/Engine/Source/Runtime/TraceLog/Public/Trace/Detail/Protocols/Protocol1.h b/Engine/Source/Runtime/TraceLog/Public/Trace/Detail/Protocols/Protocol1.h
--- a/Engine/Source/Runtime/TraceLog/Public/Trace/Detail/Protocols/Protocol1.h
+++ b/Engine/Source/Runtime/TraceLog/Public/Trace/Detail/Protocols/Protocol1.h
@@ -22,6 +22,7 @@ using Protocol0::FNewEventEvent;
 enum class EEventFlags : uint8
 {
 	Important		= 1 << 0,
+	MaybeHasAux		= 1 << 1,
 };
 
 ////////////////////////////////////////////////////////////////////////////////
@@ -43,5 +44,22 @@ struct FEventHeader
 	uint8		EventData[];
 };
 
+////////////////////////////////////////////////////////////////////////////////
+struct FAuxHeader
+{
+	enum : uint8
+	{
+		AuxDataBit	= 0x80,
+		FieldMask	= 0x7f,
+	};
+
+	union
+	{
+		uint8	FieldIndex;	// 7 bits max (MSB is used to indicate aux data)
+		uint32	Size;		// encoded as (Size & 0x00ffffff) << 8
+	};
+	uint8		Data[];
+};
+
 } // namespace Protocol1
 } // namespace Trace
diff --git a/Engine/Source/Runtime/TraceLog/Public/Trace/Detail/Trace.h b/Engine/Source/Runtime/TraceLog/Public/Trace/Detail/Trace.h
--- a/Engine/Source/Runtime/TraceLog/Public/Trace/Detail/Trace.h
+++ b/Engine/Source/Runtime/TraceLog/Public/Trace/Detail/Trace.h
@@ -27,7 +27,8 @@
 			{ \
 				const uint32 Always = Trace::FEventDef::Flag_Always; \
 				const uint32 Important = Trace::FEventDef::Flag_Important; \
-				const uint32 Flags = (0, ##__VA_ARGS__); \
+				uint32 Flags = decltype(EventProps_Private)::MaybeHasAux ? Trace::FEventDef::Flag_MaybeHasAux : 0; \
+				Flags |= (0, ##__VA_ARGS__); \
 				F##LoggerName##EventName##Fields Fields; \
 				const auto* Descs = (Trace::FFieldDesc*)(&Fields); \
 				uint32 DescCount = uint32(sizeof(Fields) / sizeof(*Descs)); \
@@ -46,8 +47,8 @@
 			decltype(FieldName)::Offset + decltype(FieldName)::Size,
 
 #define TRACE_PRIVATE_EVENT_END() \
-		Trace::EndOfFields> const EventSize_Private = {}; \
-		Trace::TField<0, decltype(EventSize_Private)::Value, Trace::Attachment> const Attachment = {}; \
+		Trace::EventProps> const EventProps_Private = {}; \
+		Trace::TField<0, decltype(EventProps_Private)::Size, Trace::Attachment> const Attachment = {}; \
 		explicit operator bool () const { return true; } \
 	};
 
@@ -58,12 +59,16 @@
 	)
 
 #define TRACE_PRIVATE_EVENT_SIZE(LoggerName, EventName) \
-	decltype(F##LoggerName##EventName##Fields::EventSize_Private)::Value
+	decltype(F##LoggerName##EventName##Fields::EventProps_Private)::Size
 
 #define TRACE_PRIVATE_LOG(LoggerName, EventName, ...) \
 	if (TRACE_PRIVATE_EVENT_IS_ENABLED(LoggerName, EventName)) \
 		if (const auto& __restrict EventName = (F##LoggerName##EventName##Fields&)LoggerName##EventName##Event) \
-			Trace::FEventDef::FLogScope(LoggerName##EventName##Event.Uid, TRACE_PRIVATE_EVENT_SIZE(LoggerName, EventName), ##__VA_ARGS__)
+			Trace::FEventDef::FLogScope( \
+				LoggerName##EventName##Event.Uid, \
+				TRACE_PRIVATE_EVENT_SIZE(LoggerName, EventName), \
+				bool(decltype(F##LoggerName##EventName##Fields::EventProps_Private)::MaybeHasAux), \
+				##__VA_ARGS__)
 
 #else
 
diff --git a/Engine/Source/Runtime/TraceLog/Public/Trace/Detail/Writer.inl b/Engine/Source/Runtime/TraceLog/Public/Trace/Detail/Writer.inl
--- a/Engine/Source/Runtime/TraceLog/Public/Trace/Detail/Writer.inl
+++ b/Engine/Source/Runtime/TraceLog/Public/Trace/Detail/Writer.inl
@@ -70,19 +70,24 @@ struct FLogInstance
 };
 
 ////////////////////////////////////////////////////////////////////////////////
-inline FLogInstance Writer_BeginLog(uint16 EventUid, uint16 Size)
+inline FLogInstance Writer_BeginLog(uint16 EventUid, uint16 Size, bool bMaybeHasAux)
 {
 	using namespace Private;
 
 	FWriteBuffer* Buffer = Writer_GetBuffer();
-	uint32 AllocSize = Size + sizeof(FEventHeader);
+	uint32 AllocSize = Size + sizeof(FEventHeader) + int(bMaybeHasAux);
 	Buffer->Cursor += AllocSize;
 	if (UNLIKELY(UPTRINT(Buffer->Cursor) > UPTRINT(Buffer)))
 	{
 		Buffer = Writer_NextBuffer(AllocSize);
 	}
 
-	uint8* Cursor = Buffer->Cursor - Size;
+	if (bMaybeHasAux)
+	{
+		Buffer->Cursor[-1] = 0;
+	}
+
+	uint8* Cursor = Buffer->Cursor - Size - int(bMaybeHasAux);
 
 	auto* Header = (uint16*)(Cursor); // FEventHeader
 	Header[-1] = uint16(AtomicIncrementRelaxed(&GLogSerial));
