Implemented dynamically sized array-type fields for trace events.

diff --git a/Engine/Source/Developer/TraceAnalysis/Private/Analysis/Engine.cpp b/Engine/Source/Developer/TraceAnalysis/Private/Analysis/Engine.cpp
--- a/Engine/Source/Developer/TraceAnalysis/Private/Analysis/Engine.cpp
+++ b/Engine/Source/Developer/TraceAnalysis/Private/Analysis/Engine.cpp
@@ -40,6 +40,7 @@ struct FAuxData
 	const uint8*	Data;
 	uint32			DataSize;
 	uint16			FieldIndex;
+	int16			FieldSizeAndType;
 };
 
 ////////////////////////////////////////////////////////////////////////////////
@@ -66,6 +67,7 @@ struct FAnalysisEngine::FDispatch
 		uint16		Size;
 		uint16		NameOffset;			// From FField ptr
 		int16		SizeAndType;		// value == byte_size, sign == float < 0 < int
+		bool		bIsArray;
 	};
 
 	uint32			GetFieldIndex(const ANSICHAR* Name) const;
@@ -315,12 +317,46 @@ IAnalyzer::FEventFieldInfo::EType IAnalyzer::FEventFieldInfo::GetType() const
 	return EType::None;
 }
 
+////////////////////////////////////////////////////////////////////////////////
+bool IAnalyzer::FEventFieldInfo::IsArray() const
+{
+	const auto* Inner = (const FAnalysisEngine::FDispatch::FField*)this;
+	return Inner->bIsArray;
+}
+
+
+
+////////////////////////////////////////////////////////////////////////////////
+uint32 IAnalyzer::FArrayReader::Num() const
+{
+	const auto* Inner = (const FAuxData*)this;
+	int32 SizeAndType = Inner->FieldSizeAndType;
+	SizeAndType = (SizeAndType < 0) ? -SizeAndType : SizeAndType;
+	return Inner->DataSize / SizeAndType;
+}
+
+////////////////////////////////////////////////////////////////////////////////
+const void* IAnalyzer::FArrayReader::GetImpl(uint32 Index, int16& SizeAndType) const
+{
+	const auto* Inner = (const FAuxData*)this;
+	SizeAndType = Inner->FieldSizeAndType;
+	uint32 Count = Num();
+	if (Index >= Count)
+	{
+		return nullptr;
+	}
+
+	SizeAndType = (SizeAndType < 0) ? -SizeAndType : SizeAndType;
+	return Inner->Data + (Index * SizeAndType);
+}
+
 
 
 ////////////////////////////////////////////////////////////////////////////////
 struct FAnalysisEngine::FEventDataInfo
 {
 	const FDispatch&	Dispatch;
+	FAuxDataCollector*	AuxCollector;
 	const uint8*		Ptr;
 	uint16				Size;
 };
@@ -332,6 +368,34 @@ const IAnalyzer::FEventTypeInfo& IAnalyzer::FEventData::GetTypeInfo() const
 	return (const FEventTypeInfo&)(Info->Dispatch);
 }
 
+////////////////////////////////////////////////////////////////////////////////
+const IAnalyzer::FArrayReader* IAnalyzer::FEventData::GetArrayImpl(const ANSICHAR* FieldName) const
+{
+	static const FAuxData EmptyAuxData = {};
+
+	const auto* Info = (const FAnalysisEngine::FEventDataInfo*)this;
+	if (Info->AuxCollector == nullptr)
+	{
+		return (IAnalyzer::FArrayReader*)&EmptyAuxData;
+	}
+
+	uint32 Index = Info->Dispatch.GetFieldIndex(FieldName);
+	if (Index >= Info->Dispatch.FieldCount)
+	{
+		return (IAnalyzer::FArrayReader*)&EmptyAuxData;
+	}
+	for (FAuxData& Data : *(Info->AuxCollector))
+	{
+		if (Data.FieldIndex == Index)
+		{
+			Data.FieldSizeAndType = Info->Dispatch.Fields[Index].SizeAndType;
+			return (IAnalyzer::FArrayReader*)&Data;
+		}
+	}
+
+	return (IAnalyzer::FArrayReader*)&EmptyAuxData;
+}
+
 ////////////////////////////////////////////////////////////////////////////////
 const void* IAnalyzer::FEventData::GetValueImpl(const ANSICHAR* FieldName, int16& SizeAndType) const
 {
@@ -671,6 +735,7 @@ void FAnalysisEngine::OnNewEventProtocol0(FDispatchBuilder& Builder, const void*
 		auto& OutField = Builder.AddField(NameCursor, Field.NameSize, Field.Size);
 		OutField.Offset = Field.Offset;
 		OutField.SizeAndType = TypeSize;
+		OutField.bIsArray = (Field.TypeInfo & Protocol0::Field_Array) != 0;
 
 		NameCursor += Field.NameSize;
 	}
@@ -828,7 +893,12 @@ bool FAnalysisEngine::OnDataProtocol0()
 			return false;
 		}
 
-		FEventDataInfo EventDataInfo = { *Dispatch, Header->EventData, Header->Size };
+		FEventDataInfo EventDataInfo = {
+			*Dispatch,
+			nullptr,
+			Header->EventData,
+			Header->Size
+		};
 		const FEventData& EventData = (FEventData&)EventDataInfo;
 
 		ForEachRoute(Dispatch, [&] (IAnalyzer* Analyzer, uint16 RouteId)
@@ -929,7 +999,12 @@ int32 FAnalysisEngine::OnDataProtocol1(FStreamReader& Reader)
 
 		++NextLogSerial;
 
-		FEventDataInfo EventDataInfo = { *Dispatch, Header->EventData, Header->Size };
+		FEventDataInfo EventDataInfo = {
+			*Dispatch,
+			&AuxCollector,
+			Header->EventData,
+			Header->Size
+		};
 		const FEventData& EventData = (FEventData&)EventDataInfo;
 
 		ForEachRoute(Dispatch, [&] (IAnalyzer* Analyzer, uint16 RouteId)
diff --git a/Engine/Source/Developer/TraceAnalysis/Public/Trace/Analyzer.h b/Engine/Source/Developer/TraceAnalysis/Public/Trace/Analyzer.h
--- a/Engine/Source/Developer/TraceAnalysis/Public/Trace/Analyzer.h
+++ b/Engine/Source/Developer/TraceAnalysis/Public/Trace/Analyzer.h
@@ -74,6 +74,9 @@ public:
 
 		/** What type of field is this? */
 		EType GetType() const;
+
+		/** Is this field an array-type field? */
+		bool IsArray() const;
 	};
 
 	struct TRACEANALYSIS_API FEventTypeInfo
@@ -98,6 +101,24 @@ public:
 		const FEventFieldInfo* GetFieldInfo(uint32 Index) const;
 	};
 
+	struct TRACEANALYSIS_API FArrayReader
+	{
+		/* Returns the number of elements in the array */
+		uint32 Num() const;
+
+	protected:
+		const void* GetImpl(uint32 Index, int16& SizeAndType) const;
+	};
+
+	template <typename ValueType>
+	struct TArrayReader
+		: public FArrayReader
+	{
+		/* Returns a element from the array an the given index or zero if Index
+		 * is out of bounds. */
+		ValueType operator [] (uint32 Index) const;
+	};
+
 	struct TRACEANALYSIS_API FEventData
 	{
 		/** Returns an object describing the underlying event's type. */
@@ -106,6 +127,12 @@ public:
 		/** Queries the value of a field of the event. */
 		template <typename ValueType> ValueType GetValue(const ANSICHAR* FieldName) const;
 
+		/** Returns an object for reading data from an array-type field. A valid
+		 * array reader object will always be return even if no field matching the
+		 * given name was found.
+		 * @param FieldName The name of the event's field to get the value for. */
+		template <typename ValueType> const TArrayReader<ValueType>& GetArray(const ANSICHAR* FieldName) const;
+
 		/** Returns the event's attachment. Not that this will always return an
 		 * address but if the event has no attachment then reading from that
 		 * address if undefined. */
@@ -119,6 +146,7 @@ public:
 
 	private:
 		const void* GetValueImpl(const ANSICHAR* FieldName, int16& SizeAndType) const;
+		const FArrayReader* GetArrayImpl(const ANSICHAR* FieldName) const;
 	};
 
 	struct FOnEventContext
@@ -161,20 +189,21 @@ public:
 	{
 		return true;
 	}
+
+private:
+	template <typename ValueType> static ValueType CoerceValue(const void* Addr, int16 SizeAndType);
 };
 
 ////////////////////////////////////////////////////////////////////////////////
 template <typename ValueType>
-ValueType IAnalyzer::FEventData::GetValue(const ANSICHAR* FieldName) const
+ValueType IAnalyzer::CoerceValue(const void* Addr, int16 SizeAndType)
 {
-	int16 FieldSizeAndType;
-	const void* Addr = GetValueImpl(FieldName, FieldSizeAndType);
 	if (Addr == nullptr)
 	{
 		return ValueType(0);
 	}
 
-	switch (FieldSizeAndType)
+	switch (SizeAndType)
 	{
 	case -4: return ValueType(*(const float*)(Addr));
 	case -8: return ValueType(*(const double*)(Addr));
@@ -187,4 +216,30 @@ ValueType IAnalyzer::FEventData::GetValue(const ANSICHAR* FieldName) const
 	return ValueType(0);
 }
 
+////////////////////////////////////////////////////////////////////////////////
+template <typename ValueType>
+ValueType IAnalyzer::FEventData::GetValue(const ANSICHAR* FieldName) const
+{
+	int16 FieldSizeAndType;
+	const void* Addr = GetValueImpl(FieldName, FieldSizeAndType);
+	return CoerceValue<ValueType>(Addr, FieldSizeAndType);
+}
+
+////////////////////////////////////////////////////////////////////////////////
+template <typename ValueType>
+const IAnalyzer::TArrayReader<ValueType>& IAnalyzer::FEventData::GetArray(const ANSICHAR* FieldName) const
+{
+	const FArrayReader* Base = GetArrayImpl(FieldName);
+	return *(TArrayReader<ValueType>*)(Base);
+}
+
+////////////////////////////////////////////////////////////////////////////////
+template <typename ValueType>
+ValueType IAnalyzer::TArrayReader<ValueType>::operator [] (uint32 Index) const
+{
+	int16 ElementSizeAndType;
+	const void* Addr = GetImpl(Index, ElementSizeAndType);
+	return CoerceValue<ValueType>(Addr, ElementSizeAndType);
+}
+
 } // namespace Trace
diff --git a/Engine/Source/Runtime/TraceLog/Public/Trace/Detail/Field.h b/Engine/Source/Runtime/TraceLog/Public/Trace/Detail/Field.h
--- a/Engine/Source/Runtime/TraceLog/Public/Trace/Detail/Field.h
+++ b/Engine/Source/Runtime/TraceLog/Public/Trace/Detail/Field.h
@@ -8,6 +8,7 @@
 
 #include "Protocol.h"
 #include "Templates/UnrealTemplate.h"
+#include "Writer.inl"
 
 namespace Trace
 {
@@ -28,6 +29,16 @@ template <> struct TFieldType<float>		{ enum { Tid = int(EFieldType::Float32),Si
 template <> struct TFieldType<double>		{ enum { Tid = int(EFieldType::Float64),Size = sizeof(double) }; };
 template <class T> struct TFieldType<T*>	{ enum { Tid = int(EFieldType::Pointer),Size = sizeof(void*) }; };
 
+template <typename T>
+struct TFieldType<T[]>
+{
+	enum
+	{
+		Tid  = int(TFieldType<T>::Tid)|int(EFieldType::Array),
+		Size = 0,
+	};
+};
+
 
 
 ////////////////////////////////////////////////////////////////////////////////
@@ -94,27 +105,68 @@ enum class EIndexPack
 		}
 
 ////////////////////////////////////////////////////////////////////////////////
-template <int Index, int Offset, typename Type> struct TField;
+template <int InIndex, int InOffset, typename Type>
+struct TField<InIndex, InOffset, Type[]>
+{
+	TRACE_PRIVATE_FIELD(InIndex|int(EIndexPack::MaybeHasAux), InOffset, Type[]);
 
+	static_assert(sizeof(Private::FWriteBuffer::Overflow) >= sizeof(FAuxHeader), "FWriteBuffer::Overflow is not large enough");
 
+	struct FActionable
+	{
+		void Write(uint8* __restrict) const {}
+	};
 
-#if 0
-////////////////////////////////////////////////////////////////////////////////
-template <typename Type>
-struct TField<Type[]>
-	: public TFieldBase<sizeof(uint32)>
-{
-	TField(const ANSICHAR (Name)[]) : TFieldBase<sizeof(uint32)>(Name) {}
-};
+	const FActionable operator () (Type const* Data, int Count) const
+	{
+		using namespace Private;
 
-////////////////////////////////////////////////////////////////////////////////
-template <int Offset, typename Type, int Count>
-struct TField<Offset, Type[Count]>
-	: public TFieldBase<Offset, Type>
-{
-	TField(const ANSICHAR (Name)[]) : TFieldBase<Offset, Type[Count]>(Name) {}
+		// Header
+		const int bMaybeHasAux = true;
+		FWriteBuffer* Buffer = Writer_GetBuffer();
+		Buffer->Cursor += sizeof(FAuxHeader) - bMaybeHasAux;
+
+		int32 Size = ((sizeof(Type) * Count) & 0x00ffffff);
+
+		auto* Header = (FAuxHeader*)(Buffer->Cursor - sizeof(FAuxHeader));
+		Header->Size = Size << 8;
+		Header->FieldIndex = uint8(0x80 | (Index & int(EIndexPack::FieldCountMask)));
+
+		// Array data
+		const uint8* ReadCursor = (uint8*)Data;
+		bool bCommit = (UPTRINT(Header) == UPTRINT(Buffer->Committed));
+		while (true)
+		{
+			if (Buffer->Cursor >= (uint8*)Buffer)
+			{
+				Buffer = Writer_NextBuffer(0);
+				bCommit = true;
+			}
+
+			int32 Remaining = int32((uint8*)Buffer - Buffer->Cursor);
+			int32 SegmentSize = (Remaining < Size) ? Remaining : Size;
+			memcpy(Buffer->Cursor, ReadCursor, SegmentSize);
+			Buffer->Cursor += SegmentSize;
+
+			if (bCommit)
+			{
+				AtomicStoreRelease<uint8* __restrict>(&(Buffer->Committed), Buffer->Cursor);
+			}
+
+			Size -= SegmentSize;
+			if (Size <= 0)
+			{
+				break;
+			}
+
+			ReadCursor += SegmentSize;
+		}
+
+		Buffer->Cursor[0] = 0;
+		Buffer->Cursor++;
+		return {};
+	}
 };
-#endif
 
 ////////////////////////////////////////////////////////////////////////////////
 template <int InIndex, int InOffset, typename Type>
diff --git a/Engine/Source/Runtime/TraceLog/Public/Trace/Detail/Protocols/Protocol0.h b/Engine/Source/Runtime/TraceLog/Public/Trace/Detail/Protocols/Protocol0.h
--- a/Engine/Source/Runtime/TraceLog/Public/Trace/Detail/Protocols/Protocol0.h
+++ b/Engine/Source/Runtime/TraceLog/Public/Trace/Detail/Protocols/Protocol0.h
@@ -21,6 +21,7 @@ enum : uint8
 	Field_CategoryMask	= 0300,
 	Field_Integer		= 0000, 
 	Field_Float			= 0100,
+	Field_Array			= 0200,
 
 	/* Size */
 	Field_Pow2SizeMask	= 0003,
@@ -46,6 +47,7 @@ enum class EFieldType : uint8
 	Pointer	= Field_Integer | Field_Ptr,
 	Float32	= Field_Float   | Field_32,
 	Float64	= Field_Float   | Field_64,
+	Array	= Field_Array,
 };
 
 ////////////////////////////////////////////////////////////////////////////////
