Updated the documentation comments for FEventData::GetValue().

#rnx

diff --git a/Engine/Source/Developer/TraceAnalysis/Public/Trace/Analyzer.h b/Engine/Source/Developer/TraceAnalysis/Public/Trace/Analyzer.h
--- a/Engine/Source/Developer/TraceAnalysis/Public/Trace/Analyzer.h
+++ b/Engine/Source/Developer/TraceAnalysis/Public/Trace/Analyzer.h
@@ -124,7 +124,10 @@ public:
 		/** Returns an object describing the underlying event's type. */
 		const FEventTypeInfo& GetTypeInfo() const;
 
-		/** Queries the value of a field of the event. */
+		/** Queries the value of a field of the event. It is not necessary to match
+		 * ValueType to the type in the event.
+		 * @param FieldName The name of the event's field to get the value for.
+		 * @return Value of the field (coerced to ValueType) if found, otherwise 0. */
 		template <typename ValueType> ValueType GetValue(const ANSICHAR* FieldName) const;
 
 		/** Returns an object for reading data from an array-type field. A valid
