Expose Trace's internal ThreadId to analyzers. Note that this will not the same as the OS' thread id.

#rnx
#ushell-unshelve of 10632436

diff --git a/Engine/Source/Developer/TraceAnalysis/Private/Analysis/Engine.cpp b/Engine/Source/Developer/TraceAnalysis/Private/Analysis/Engine.cpp
--- a/Engine/Source/Developer/TraceAnalysis/Private/Analysis/Engine.cpp
+++ b/Engine/Source/Developer/TraceAnalysis/Private/Analysis/Engine.cpp
@@ -903,7 +903,7 @@ bool FAnalysisEngine::OnDataProtocol0()
 
 		ForEachRoute(Dispatch, [&] (IAnalyzer* Analyzer, uint16 RouteId)
 		{
-			if (!Analyzer->OnEvent(RouteId, { SessionContext, EventData }))
+			if (!Analyzer->OnEvent(RouteId, { SessionContext, EventData, ~0u }))
 			{
 				RetireAnalyzer(Analyzer);
 			}
@@ -928,7 +928,8 @@ bool FAnalysisEngine::OnDataProtocol1()
 		FTidPacketTransport::ThreadIter Iter = InnerTransport->ReadThreads();
 		while (FStreamReader* Reader = InnerTransport->GetNextThread(Iter))
 		{
-			int32 ThreadEventCount = OnDataProtocol1(*Reader);
+			uint32 ThreadId = InnerTransport->GetThreadId(Iter);
+			int32 ThreadEventCount = OnDataProtocol1(ThreadId, *Reader);
 			if (ThreadEventCount < 0)
 			{
 				return false;
@@ -943,7 +944,7 @@ bool FAnalysisEngine::OnDataProtocol1()
 }
 
 ////////////////////////////////////////////////////////////////////////////////
-int32 FAnalysisEngine::OnDataProtocol1(FStreamReader& Reader)
+int32 FAnalysisEngine::OnDataProtocol1(uint32 ThreadId, FStreamReader& Reader)
 {
 	int32 EventCount = 0;
 	while (!Reader.IsEmpty())
@@ -1009,7 +1010,7 @@ int32 FAnalysisEngine::OnDataProtocol1(FStreamReader& Reader)
 
 		ForEachRoute(Dispatch, [&] (IAnalyzer* Analyzer, uint16 RouteId)
 		{
-			if (!Analyzer->OnEvent(RouteId, { SessionContext, EventData }))
+			if (!Analyzer->OnEvent(RouteId, { SessionContext, EventData, ThreadId }))
 			{
 				RetireAnalyzer(Analyzer);
 			}
diff --git a/Engine/Source/Developer/TraceAnalysis/Private/Analysis/Engine.h b/Engine/Source/Developer/TraceAnalysis/Private/Analysis/Engine.h
--- a/Engine/Source/Developer/TraceAnalysis/Private/Analysis/Engine.h
+++ b/Engine/Source/Developer/TraceAnalysis/Private/Analysis/Engine.h
@@ -46,7 +46,7 @@ private:
 	bool				EstablishTransport(FStreamReader& Reader);
 	bool				OnDataProtocol0();
 	bool				OnDataProtocol1();
-	int32				OnDataProtocol1(FStreamReader& Reader);
+	int32				OnDataProtocol1(uint32 ThreadId, FStreamReader& Reader);
 	int32				OnDataProtocol1Aux(FStreamReader& Reader, FAuxDataCollector& Collector);
 	bool				AddDispatch(FDispatch* Dispatch);
 	template <typename ImplType>
diff --git a/Engine/Source/Developer/TraceAnalysis/Private/Analysis/Transport/TidPacketTransport.cpp b/Engine/Source/Developer/TraceAnalysis/Private/Analysis/Transport/TidPacketTransport.cpp
--- a/Engine/Source/Developer/TraceAnalysis/Private/Analysis/Transport/TidPacketTransport.cpp
+++ b/Engine/Source/Developer/TraceAnalysis/Private/Analysis/Transport/TidPacketTransport.cpp
@@ -106,4 +106,16 @@ FStreamReader* FTidPacketTransport::GetNextThread(ThreadIter& Iter)
 	return nullptr;
 }
 
+////////////////////////////////////////////////////////////////////////////////
+int32 FTidPacketTransport::GetThreadId(ThreadIter Iter) const
+{
+	uint32 Index = Iter - 1;
+	if (Index >= uint32(Threads.Num()))
+	{
+		return -1;
+	}
+
+	return Threads[Index].ThreadId;
+}
+
 } // namespace Trace
diff --git a/Engine/Source/Developer/TraceAnalysis/Private/Analysis/Transport/TidPacketTransport.h b/Engine/Source/Developer/TraceAnalysis/Private/Analysis/Transport/TidPacketTransport.h
--- a/Engine/Source/Developer/TraceAnalysis/Private/Analysis/Transport/TidPacketTransport.h
+++ b/Engine/Source/Developer/TraceAnalysis/Private/Analysis/Transport/TidPacketTransport.h
@@ -21,6 +21,7 @@ public:
 	void					Update();
 	ThreadIter				ReadThreads();
 	FStreamReader*			GetNextThread(ThreadIter& Iter);
+	int32					GetThreadId(ThreadIter Iter) const;
 
 private:
 	struct FThreadStream
diff --git a/Engine/Source/Developer/TraceAnalysis/Public/Trace/Analyzer.h b/Engine/Source/Developer/TraceAnalysis/Public/Trace/Analyzer.h
--- a/Engine/Source/Developer/TraceAnalysis/Public/Trace/Analyzer.h
+++ b/Engine/Source/Developer/TraceAnalysis/Public/Trace/Analyzer.h
@@ -156,6 +156,7 @@ public:
 	{
 		const FSessionContext&	SessionContext;
 		const FEventData&		EventData;
+		uint32					ThreadId;
 	};
 
 	virtual ~IAnalyzer() = default;
