Assigned more bits to the event serial numbers to prevent confusion during analysis when it wraps around.

#rnx
#ushell-unshelve of 10634350

diff --git a/Engine/Source/Developer/TraceAnalysis/Private/Analysis/Engine.cpp b/Engine/Source/Developer/TraceAnalysis/Private/Analysis/Engine.cpp
--- a/Engine/Source/Developer/TraceAnalysis/Private/Analysis/Engine.cpp
+++ b/Engine/Source/Developer/TraceAnalysis/Private/Analysis/Engine.cpp
@@ -652,8 +652,14 @@ void FAnalysisEngine::OnNewEventInternal(const FOnEventContext& Context)
 	FDispatchBuilder Builder;
 	switch (ProtocolVersion)
 	{
-	case Protocol0::EProtocol::Id: OnNewEventProtocol0(Builder, EventData.Ptr); break;
-	case Protocol1::EProtocol::Id: OnNewEventProtocol1(Builder, EventData.Ptr); break;
+	case Protocol0::EProtocol::Id:
+		OnNewEventProtocol0(Builder, EventData.Ptr);
+		break;
+
+	case Protocol1::EProtocol::Id:
+	case Protocol2::EProtocol::Id:
+		OnNewEventProtocol1(Builder, EventData.Ptr);
+		break;
 	}
 
 	// Get the dispatch and add it into the dispatch table. Fail gently if there
@@ -814,7 +820,8 @@ bool FAnalysisEngine::EstablishTransport(FStreamReader& Reader)
 		break;
 
 	case Protocol1::EProtocol::Id:
-		ProtocolHandler = &FAnalysisEngine::OnDataProtocol1;
+	case Protocol2::EProtocol::Id:
+		ProtocolHandler = &FAnalysisEngine::OnDataProtocol2;
 		{
 			FDispatchBuilder Builder;
 			Builder.SetUid(uint16(Protocol0::EKnownEventUids::NewEvent));
@@ -916,7 +923,7 @@ bool FAnalysisEngine::OnDataProtocol0()
 }
 
 ////////////////////////////////////////////////////////////////////////////////
-bool FAnalysisEngine::OnDataProtocol1()
+bool FAnalysisEngine::OnDataProtocol2()
 {
 	auto* InnerTransport = (FTidPacketTransport*)Transport;
 	InnerTransport->Update();
@@ -929,7 +936,7 @@ bool FAnalysisEngine::OnDataProtocol1()
 		while (FStreamReader* Reader = InnerTransport->GetNextThread(Iter))
 		{
 			uint32 ThreadId = InnerTransport->GetThreadId(Iter);
-			int32 ThreadEventCount = OnDataProtocol1(ThreadId, *Reader);
+			int32 ThreadEventCount = OnDataProtocol2(ThreadId, *Reader);
 			if (ThreadEventCount < 0)
 			{
 				return false;
@@ -944,26 +951,46 @@ bool FAnalysisEngine::OnDataProtocol1()
 }
 
 ////////////////////////////////////////////////////////////////////////////////
-int32 FAnalysisEngine::OnDataProtocol1(uint32 ThreadId, FStreamReader& Reader)
+int32 FAnalysisEngine::OnDataProtocol2(uint32 ThreadId, FStreamReader& Reader)
 {
 	int32 EventCount = 0;
 	while (!Reader.IsEmpty())
 	{
 		auto Mark = Reader.SaveMark();
 
-		const auto* Header = Reader.GetPointer<Protocol1::FEventHeader>();
+		const auto* Header = Reader.GetPointer<Protocol2::FEventHeader>();
 		if (Header == nullptr)
 		{
-			break;
+			return -1;
 		}
 
-		// Make sure we consume events in the correct order
-		if (Header->Serial != uint16(NextLogSerial))
+		uint32 BlockSize = Header->Size;
+		
+		switch (ProtocolVersion)
 		{
-			break;
+			case Protocol1::EProtocol::Id:
+				{
+					const auto* HeaderV1 = (Protocol1::FEventHeader*)Header;
+					if (HeaderV1->Serial != uint16(NextLogSerial))
+					{
+						return EventCount;
+					}
+					BlockSize += sizeof(*HeaderV1);
+				}
+				break;
+
+			case Protocol2::EProtocol::Id:
+				{
+					uint32 EventSerial = Header->SerialLow|(uint32(Header->SerialHigh) << 16);
+					if (EventSerial != (NextLogSerial & 0x00ffffff))
+					{
+						return EventCount;
+					}
+					BlockSize += sizeof(*Header);
+				}
+				break;
 		}
 
-		uint32 BlockSize = Header->Size + sizeof(Protocol1::FEventHeader);
 		if (Reader.GetPointer(BlockSize) == nullptr)
 		{
 			break;
@@ -986,7 +1013,7 @@ int32 FAnalysisEngine::OnDataProtocol1(uint32 ThreadId, FStreamReader& Reader)
 		FAuxDataCollector AuxCollector;
 		if (Dispatch->Flags & FDispatch::Flag_MaybeHasAux)
 		{
-			int AuxStatus = OnDataProtocol1Aux(Reader, AuxCollector);
+			int AuxStatus = OnDataProtocol2Aux(Reader, AuxCollector);
 			if (AuxStatus == 0)
 			{
 				Reader.RestoreMark(Mark);
@@ -1003,8 +1030,8 @@ int32 FAnalysisEngine::OnDataProtocol1(uint32 ThreadId, FStreamReader& Reader)
 		FEventDataInfo EventDataInfo = {
 			*Dispatch,
 			&AuxCollector,
-			Header->EventData,
-			Header->Size
+			(const uint8*)Header + BlockSize - Header->Size,
+			Header->Size,
 		};
 		const FEventData& EventData = (FEventData&)EventDataInfo;
 
@@ -1023,7 +1050,7 @@ int32 FAnalysisEngine::OnDataProtocol1(uint32 ThreadId, FStreamReader& Reader)
 }
 
 ////////////////////////////////////////////////////////////////////////////////
-int32 FAnalysisEngine::OnDataProtocol1Aux(FStreamReader& Reader, FAuxDataCollector& Collector)
+int32 FAnalysisEngine::OnDataProtocol2Aux(FStreamReader& Reader, FAuxDataCollector& Collector)
 {
 	while (true)
 	{
diff --git a/Engine/Source/Developer/TraceAnalysis/Private/Analysis/Engine.h b/Engine/Source/Developer/TraceAnalysis/Private/Analysis/Engine.h
--- a/Engine/Source/Developer/TraceAnalysis/Private/Analysis/Engine.h
+++ b/Engine/Source/Developer/TraceAnalysis/Private/Analysis/Engine.h
@@ -45,9 +45,9 @@ private:
 
 	bool				EstablishTransport(FStreamReader& Reader);
 	bool				OnDataProtocol0();
-	bool				OnDataProtocol1();
-	int32				OnDataProtocol1(uint32 ThreadId, FStreamReader& Reader);
-	int32				OnDataProtocol1Aux(FStreamReader& Reader, FAuxDataCollector& Collector);
+	bool				OnDataProtocol2();
+	int32				OnDataProtocol2(uint32 ThreadId, FStreamReader& Reader);
+	int32				OnDataProtocol2Aux(FStreamReader& Reader, FAuxDataCollector& Collector);
 	bool				AddDispatch(FDispatch* Dispatch);
 	template <typename ImplType>
 	void				ForEachRoute(const FDispatch* Dispatch, ImplType&& Impl);
@@ -60,7 +60,7 @@ private:
 	TArray<FDispatch*>	Dispatches;
 	class FTransport*	Transport = nullptr;
 	ProtocolHandlerType	ProtocolHandler = nullptr;
-	uint16				NextLogSerial = 0;
+	uint32				NextLogSerial = 0;
 	uint8				ProtocolVersion = 0;
 };
 
diff --git a/Engine/Source/Runtime/TraceLog/Public/Trace/Detail/Protocol.h b/Engine/Source/Runtime/TraceLog/Public/Trace/Detail/Protocol.h
--- a/Engine/Source/Runtime/TraceLog/Public/Trace/Detail/Protocol.h
+++ b/Engine/Source/Runtime/TraceLog/Public/Trace/Detail/Protocol.h
@@ -6,7 +6,7 @@
 
 ////////////////////////////////////////////////////////////////////////////////
 #if UE_TRACE_ENABLED
-#	define TRACE_PRIVATE_PROTOCOL_1
+#	define TRACE_PRIVATE_PROTOCOL_2
 #endif
 
 #if defined(_MSC_VER)
@@ -16,6 +16,7 @@
 
 #include "Protocols/Protocol0.h"
 #include "Protocols/Protocol1.h"
+#include "Protocols/Protocol2.h"
 
 #if defined(_MSC_VER)
 	#pragma warning(pop)
diff --git a/Engine/Source/Runtime/TraceLog/Public/Trace/Detail/Protocols/Protocol2.h b/Engine/Source/Runtime/TraceLog/Public/Trace/Detail/Protocols/Protocol2.h
new file mode 100644
--- /dev/null
+++ b/Engine/Source/Runtime/TraceLog/Public/Trace/Detail/Protocols/Protocol2.h
@@ -0,0 +1,38 @@
+// Copyright 1998-2019 Epic Games, Inc. All Rights Reserved.
+
+#pragma once
+
+namespace Trace
+{
+
+#if defined(TRACE_PRIVATE_PROTOCOL_2)
+inline
+#endif
+namespace Protocol2
+{
+
+////////////////////////////////////////////////////////////////////////////////
+enum EProtocol : uint8 { Id = 2 };
+
+////////////////////////////////////////////////////////////////////////////////
+using Protocol1::EFieldType;
+using Protocol1::FNewEventEvent;
+using Protocol1::EEventFlags;
+using Protocol1::EKnownEventUids;
+using Protocol1::FAuxHeader;
+
+////////////////////////////////////////////////////////////////////////////////
+#pragma pack(push, 1)
+struct FEventHeader
+{
+	uint16		Uid;
+	uint16		Size;
+	uint16		SerialLow;		// 24-bit...
+	uint8		SerialHigh;		// ...serial no.
+	uint8		EventData[];
+};
+#pragma pack(pop)
+static_assert(sizeof(FEventHeader) == 7, "Packing assumption doesn't hold");
+
+} // namespace Protocol1
+} // namespace Trace
diff --git a/Engine/Source/Runtime/TraceLog/Public/Trace/Detail/Writer.inl b/Engine/Source/Runtime/TraceLog/Public/Trace/Detail/Writer.inl
--- a/Engine/Source/Runtime/TraceLog/Public/Trace/Detail/Writer.inl
+++ b/Engine/Source/Runtime/TraceLog/Public/Trace/Detail/Writer.inl
@@ -92,8 +92,9 @@ inline FLogInstance Writer_BeginLog(uint16 EventUid, uint16 Size, bool bMaybeHas
 
 	uint8* Cursor = Buffer->Cursor - Size - int(bMaybeHasAux);
 
-	auto* Header = (uint16*)(Cursor); // FEventHeader
-	Header[-1] = uint16(AtomicIncrementRelaxed(&GLogSerial));
+	// Event header
+	auto* Header = (uint16*)(Cursor - sizeof(FEventHeader::SerialHigh)); // FEventHeader1
+	*(uint32*)(Header - 1) = uint32(AtomicIncrementRelaxed(&GLogSerial));
 	Header[-2] = Size;
 	Header[-3] = EventUid;
 
