Fixed incorrect return value if no header's found. That's not a failure.

#rnx

diff --git a/Engine/Source/Developer/TraceAnalysis/Private/Analysis/Engine.cpp b/Engine/Source/Developer/TraceAnalysis/Private/Analysis/Engine.cpp
--- a/Engine/Source/Developer/TraceAnalysis/Private/Analysis/Engine.cpp
+++ b/Engine/Source/Developer/TraceAnalysis/Private/Analysis/Engine.cpp
@@ -961,7 +961,7 @@ int32 FAnalysisEngine::OnDataProtocol2(uint32 ThreadId, FStreamReader& Reader)
 		const auto* Header = Reader.GetPointer<Protocol2::FEventHeader>();
 		if (Header == nullptr)
 		{
-			return -1;
+			break;
 		}
 
 		uint32 BlockSize = Header->Size;
