Validate an event's ID before doing anything else with it.

#rnx

diff --git a/Engine/Source/Developer/TraceAnalysis/Private/Analysis/Engine.cpp b/Engine/Source/Developer/TraceAnalysis/Private/Analysis/Engine.cpp
--- a/Engine/Source/Developer/TraceAnalysis/Private/Analysis/Engine.cpp
+++ b/Engine/Source/Developer/TraceAnalysis/Private/Analysis/Engine.cpp
@@ -964,6 +964,18 @@ int32 FAnalysisEngine::OnDataProtocol2(uint32 ThreadId, FStreamReader& Reader)
 			break;
 		}
 
+		uint16 Uid = Header->Uid & uint16(Protocol2::EKnownEventUids::UidMask);
+		if (Uid >= Dispatches.Num())
+		{
+			return -1;
+		}
+
+		const FDispatch* Dispatch = Dispatches[Uid];
+		if (Dispatch == nullptr)
+		{
+			return -1;
+		}
+
 		uint32 BlockSize = Header->Size;
 		
 		switch (ProtocolVersion)
@@ -996,18 +1008,6 @@ int32 FAnalysisEngine::OnDataProtocol2(uint32 ThreadId, FStreamReader& Reader)
 			break;
 		}
 
-		uint16 Uid = Header->Uid & uint16(Protocol1::EKnownEventUids::UidMask);
-		if (Uid >= Dispatches.Num())
-		{
-			return -1;
-		}
-
-		const FDispatch* Dispatch = Dispatches[Uid];
-		if (Dispatch == nullptr)
-		{
-			return -1;
-		}
-
 		Reader.Advance(BlockSize);
 
 		FAuxDataCollector AuxCollector;
