First draft of a server/client architecture for the trace store. Currently disabled pending TPS review of dependencies.

#rnx

diff --git a/Engine/Source/Developer/TraceAnalysis/Private/Asio/Asio.cpp b/Engine/Source/Developer/TraceAnalysis/Private/Asio/Asio.cpp
new file mode 100644
--- /dev/null
+++ b/Engine/Source/Developer/TraceAnalysis/Private/Asio/Asio.cpp
@@ -0,0 +1,35 @@
+// Copyright 1998-2019 Epic Games, Inc. All Rights Reserved.
+
+#include "Asio.h"
+
+#if TRACE_WITH_ASIO
+
+#include "HAL/Platform.h"
+#include "Misc/AssertionMacros.h"
+
+#if PLATFORM_WINDOWS
+#	include "Windows/AllowWindowsPlatformTypes.h"
+#	include "Windows/AllowWindowsPlatformAtomics.h"
+#endif
+
+#if defined(_MSC_VER)
+#	pragma warning(push)
+#	pragma warning(disable : 4191)
+#endif
+
+THIRD_PARTY_INCLUDES_START
+
+#include "asio/impl/src.hpp"
+
+THIRD_PARTY_INCLUDES_END
+
+#if defined(_MSC_VER)
+#	pragma warning(pop)
+#endif
+
+#if PLATFORM_WINDOWS
+#	include "Windows/HideWindowsPlatformAtomics.h"
+#	include "Windows/HideWindowsPlatformTypes.h"
+#endif
+
+#endif // TRACE_WITH_ASIO
diff --git a/Engine/Source/Developer/TraceAnalysis/Private/Asio/Asio.h b/Engine/Source/Developer/TraceAnalysis/Private/Asio/Asio.h
new file mode 100644
--- /dev/null
+++ b/Engine/Source/Developer/TraceAnalysis/Private/Asio/Asio.h
@@ -0,0 +1,48 @@
+// Copyright Epic Games, Inc. All Rights Reserved.
+
+#pragma once
+
+#if !defined(TRACE_WITH_ASIO)
+#	define TRACE_WITH_ASIO 0
+#endif // 0
+
+#if TRACE_WITH_ASIO
+
+#if PLATFORM_WINDOWS
+#	include "Windows/AllowWindowsPlatformTypes.h"
+#	include "Windows/AllowWindowsPlatformAtomics.h"
+#endif
+
+THIRD_PARTY_INCLUDES_START
+
+namespace asio {
+namespace detail {
+
+template <typename ExceptionType>
+void throw_exception(const ExceptionType& Exception)
+{
+	/* Intentionally blank */
+}
+
+} // namespace detail
+} // namespace asio
+
+#include "asio/connect.hpp"
+#include "asio/io_context.hpp"
+#include "asio/ip/tcp.hpp"
+#include "asio/read.hpp"
+#include "asio/write.hpp"
+#include "asio/write_at.hpp"
+
+#if PLATFORM_WINDOWS
+#	include "asio/windows/random_access_handle.hpp"
+#endif
+
+THIRD_PARTY_INCLUDES_END
+
+#if PLATFORM_WINDOWS
+#	include "Windows/HideWindowsPlatformAtomics.h"
+#	include "Windows/HideWindowsPlatformTypes.h"
+#endif
+
+#endif // TRACE_WITH_ASIO
diff --git a/Engine/Source/Developer/TraceAnalysis/Private/Store/AsioContext.cpp b/Engine/Source/Developer/TraceAnalysis/Private/Store/AsioContext.cpp
new file mode 100644
--- /dev/null
+++ b/Engine/Source/Developer/TraceAnalysis/Private/Store/AsioContext.cpp
@@ -0,0 +1,54 @@
+// Copyright Epic Games, Inc. All Rights Reserved.
+
+#include "AsioContext.h"
+
+#if TRACE_WITH_ASIO
+
+namespace Trace
+{
+
+////////////////////////////////////////////////////////////////////////////////
+FAsioContext::FAsioContext(int32 ThreadCount)
+{
+	if (ThreadCount < 1)
+	{
+		ThreadCount = 1;
+	}
+
+	IoContext = new asio::io_context(ThreadCount);
+
+	ThreadPool.SetNum(ThreadCount);
+}
+
+////////////////////////////////////////////////////////////////////////////////
+void FAsioContext::Start()
+{
+	for (std::thread& Thread : ThreadPool)
+	{
+		std::thread TempThread([this] () { IoContext->run(); });
+		Thread = MoveTemp(TempThread);
+	}
+}
+
+////////////////////////////////////////////////////////////////////////////////
+FAsioContext::~FAsioContext()
+{
+	IoContext->stop();
+
+	for (std::thread& Thread : ThreadPool)
+	{
+		Thread.join();
+	}
+
+	delete IoContext;
+}
+
+////////////////////////////////////////////////////////////////////////////////
+asio::io_context& FAsioContext::Get()
+{
+	return *IoContext;
+}
+
+} // namespace Trace
+
+#endif // TRACE_WITH_ASIO
diff --git a/Engine/Source/Developer/TraceAnalysis/Private/Store/AsioContext.h b/Engine/Source/Developer/TraceAnalysis/Private/Store/AsioContext.h
new file mode 100644
--- /dev/null
+++ b/Engine/Source/Developer/TraceAnalysis/Private/Store/AsioContext.h
@@ -0,0 +1,30 @@
+// Copyright Epic Games, Inc. All Rights Reserved.
+
+#pragma once
+
+#include "Asio/Asio.h"
+
+#if TRACE_WITH_ASIO
+
+#include <thread>
+
+namespace Trace
+{
+
+////////////////////////////////////////////////////////////////////////////////
+class FAsioContext
+{
+public:
+							FAsioContext(int32 ThreadCount);
+							~FAsioContext();
+	asio::io_context&		Get();
+	void					Start();
+
+private:
+	asio::io_context*		IoContext;
+	TArray<std::thread>		ThreadPool;
+};
+
+} // namespace Trace
+
+#endif // TRACE_WITH_ASIO
diff --git a/Engine/Source/Developer/TraceAnalysis/Private/Store/AsioFile.cpp b/Engine/Source/Developer/TraceAnalysis/Private/Store/AsioFile.cpp
new file mode 100644
--- /dev/null
+++ b/Engine/Source/Developer/TraceAnalysis/Private/Store/AsioFile.cpp
@@ -0,0 +1,99 @@
+// Copyright Epic Games, Inc. All Rights Reserved.
+
+#include "AsioFile.h"
+
+#if TRACE_WITH_ASIO
+
+namespace Trace
+{
+
+////////////////////////////////////////////////////////////////////////////////
+FAsioFile::FAsioFile(asio::io_context& IoContext, UPTRINT OsHandle)
+: Handle(IoContext, HANDLE(OsHandle))
+{
+}
+
+////////////////////////////////////////////////////////////////////////////////
+void FAsioFile::Close()
+{
+	Handle.close();
+}
+
+////////////////////////////////////////////////////////////////////////////////
+bool FAsioFile::Write(const void* Src, uint32 Size, FAsioIoSink* Sink, uint32 Id)
+{
+	if (!SetSink(Sink, Id))
+	{
+		return false;
+	}
+
+	asio::async_write_at(
+		Handle,
+		Offset,
+		asio::buffer(Src, Size),
+		[this] (const asio::error_code& ErrorCode, size_t BytesWritten)
+		{
+			Offset += BytesWritten;
+			return OnIoComplete(ErrorCode, BytesWritten);
+		}
+	);
+
+	return true;
+}
+
+////////////////////////////////////////////////////////////////////////////////
+bool FAsioFile::Read(void* Dest, uint32 Size, FAsioIoSink* Sink, uint32 Id)
+{
+	return ReadSome(Dest, Size, Sink, Id);
+}
+
+////////////////////////////////////////////////////////////////////////////////
+bool FAsioFile::ReadSome(void* Dest, uint32 DestSize, FAsioIoSink* Sink, uint32 Id)
+{
+	if (!SetSink(Sink, Id))
+	{
+		return false;
+	}
+
+	Handle.async_read_some_at(
+		Offset,
+		asio::buffer(Dest, DestSize),
+		[this] (const asio::error_code& ErrorCode, size_t BytesRead)
+		{
+			Offset += BytesRead;
+			return OnIoComplete(ErrorCode, BytesRead);
+		}
+	);
+
+	return true;
+}
+
+////////////////////////////////////////////////////////////////////////////////
+FAsioWriteable* FAsioFile::WriteFile(asio::io_context& IoContext, const TCHAR* Path)
+{
+	HANDLE Handle = CreateFileW(Path, GENERIC_WRITE, FILE_SHARE_READ, nullptr,
+		CREATE_ALWAYS, FILE_FLAG_OVERLAPPED|FILE_ATTRIBUTE_NORMAL, nullptr);
+	if (Handle == INVALID_HANDLE_VALUE)
+	{
+		return nullptr;
+	}
+
+	return new FAsioFile(IoContext, UPTRINT(Handle));
+}
+
+////////////////////////////////////////////////////////////////////////////////
+FAsioReadable* FAsioFile::ReadFile(asio::io_context& IoContext, const TCHAR* Path)
+{
+	HANDLE Handle = CreateFileW(Path, GENERIC_READ, FILE_SHARE_READ|FILE_SHARE_WRITE,
+		nullptr, OPEN_EXISTING, FILE_FLAG_OVERLAPPED|FILE_ATTRIBUTE_NORMAL, nullptr);
+	if (Handle == INVALID_HANDLE_VALUE)
+	{
+		return nullptr;
+	}
+
+	return new FAsioFile(IoContext, UPTRINT(Handle));
+}
+
+} // namespace Trace
+
+#endif // TRACE_WITH_ASIO
diff --git a/Engine/Source/Developer/TraceAnalysis/Private/Store/AsioFile.h b/Engine/Source/Developer/TraceAnalysis/Private/Store/AsioFile.h
new file mode 100644
--- /dev/null
+++ b/Engine/Source/Developer/TraceAnalysis/Private/Store/AsioFile.h
@@ -0,0 +1,35 @@
+// Copyright Epic Games, Inc. All Rights Reserved.
+
+#pragma once
+
+#include "Asio/Asio.h"
+
+#if TRACE_WITH_ASIO
+
+#include "AsioIoable.h"
+
+namespace Trace
+{
+
+////////////////////////////////////////////////////////////////////////////////
+class FAsioFile
+	: public FAsioReadable
+	, public FAsioWriteable
+{
+public:
+										FAsioFile(asio::io_context& IoContext, UPTRINT OsHandle);
+	static FAsioWriteable*				WriteFile(asio::io_context& IoContext, const TCHAR* Path);
+	static FAsioReadable*				ReadFile(asio::io_context& IoContext, const TCHAR* Path);
+	virtual void						Close() override;
+	virtual bool						Write(const void* Src, uint32 Size, FAsioIoSink* Sink, uint32 Id) override;
+	virtual bool						Read(void* Dest, uint32 Size, FAsioIoSink* Sink, uint32 Id) override;
+	virtual bool						ReadSome(void* Dest, uint32 BufferSize, FAsioIoSink* Sink, uint32 Id) override;
+
+private:
+	asio::windows::random_access_handle	Handle;
+	uint64								Offset = 0;
+};
+
+} // namespace Trace
+
+#endif // TRACE_WITH_ASIO
diff --git a/Engine/Source/Developer/TraceAnalysis/Private/Store/AsioIoable.cpp b/Engine/Source/Developer/TraceAnalysis/Private/Store/AsioIoable.cpp
new file mode 100644
--- /dev/null
+++ b/Engine/Source/Developer/TraceAnalysis/Private/Store/AsioIoable.cpp
@@ -0,0 +1,43 @@
+// Copyright Epic Games, Inc. All Rights Reserved.
+
+#include "AsioIoable.h"
+
+#if TRACE_WITH_ASIO
+
+namespace Trace
+{
+
+////////////////////////////////////////////////////////////////////////////////
+bool FAsioIoable::SetSink(FAsioIoSink* Ptr, uint32 Id)
+{
+	if (SinkPtr != nullptr)
+	{
+		return false;
+	}
+
+	SinkPtr = Ptr;
+	SinkId = Id;
+	return true;
+}
+
+////////////////////////////////////////////////////////////////////////////////
+void FAsioIoable::OnIoComplete(const asio::error_code& ErrorCode, int32 Size)
+{
+	if (ErrorCode)
+	{
+		Size = -1;
+	}
+
+	if (SinkPtr == nullptr)
+	{
+		return;
+	}
+
+	FAsioIoSink* Ptr = SinkPtr;
+	SinkPtr = nullptr;
+	Ptr->OnIoComplete(SinkId, Size);
+}
+
+} // namespace Trace
+
+#endif // TRACE_WITH_ASIO
diff --git a/Engine/Source/Developer/TraceAnalysis/Private/Store/AsioIoable.h b/Engine/Source/Developer/TraceAnalysis/Private/Store/AsioIoable.h
new file mode 100644
--- /dev/null
+++ b/Engine/Source/Developer/TraceAnalysis/Private/Store/AsioIoable.h
@@ -0,0 +1,56 @@
+// Copyright Epic Games, Inc. All Rights Reserved.
+
+#pragma once
+
+#include "Asio/Asio.h"
+
+#if TRACE_WITH_ASIO
+
+namespace Trace
+{
+
+////////////////////////////////////////////////////////////////////////////////
+class FAsioIoSink
+{
+public:
+	virtual void	OnIoComplete(uint32 Id, int32 Size) = 0;
+};
+
+////////////////////////////////////////////////////////////////////////////////
+class FAsioIoable
+{
+public:
+	virtual			~FAsioIoable() = default;
+	virtual void	Close() = 0;
+
+protected:
+	bool			SetSink(FAsioIoSink* Ptr, uint32 Id);
+	void			OnIoComplete(const asio::error_code& ErrorCode, int32 Size);
+
+private:
+	FAsioIoSink*	SinkPtr = nullptr;
+	uint32			SinkId;
+};
+
+////////////////////////////////////////////////////////////////////////////////
+class FAsioReadable
+	: public virtual FAsioIoable
+{
+public:
+	virtual bool	Read(void* Dest, uint32 Size, FAsioIoSink* Sink, uint32 Id) = 0;
+	virtual bool	ReadSome(void* Dest, uint32 DestSize, FAsioIoSink* Sink, uint32 Id) = 0;
+};
+
+
+
+////////////////////////////////////////////////////////////////////////////////
+class FAsioWriteable
+	: public virtual FAsioIoable
+{
+public:
+	virtual bool	Write(const void* Src, uint32 Size, FAsioIoSink* Sink, uint32 Id) = 0;
+};
+
+} // namespace Trace
+
+#endif // TRACE_WITH_ASIO
diff --git a/Engine/Source/Developer/TraceAnalysis/Private/Store/AsioObject.cpp b/Engine/Source/Developer/TraceAnalysis/Private/Store/AsioObject.cpp
new file mode 100644
--- /dev/null
+++ b/Engine/Source/Developer/TraceAnalysis/Private/Store/AsioObject.cpp
@@ -0,0 +1,24 @@
+// Copyright Epic Games, Inc. All Rights Reserved.
+
+#include "AsioObject.h"
+
+#if TRACE_WITH_ASIO
+
+namespace Trace
+{
+
+////////////////////////////////////////////////////////////////////////////////
+FAsioObject::FAsioObject(asio::io_context& InIoContext)
+: IoContext(InIoContext)
+{
+}
+
+////////////////////////////////////////////////////////////////////////////////
+asio::io_context& FAsioObject::GetIoContext()
+{
+	return IoContext;
+}
+
+} // namespace Trace
+
+#endif // TRACE_WITH_ASIO
diff --git a/Engine/Source/Developer/TraceAnalysis/Private/Store/AsioObject.h b/Engine/Source/Developer/TraceAnalysis/Private/Store/AsioObject.h
new file mode 100644
--- /dev/null
+++ b/Engine/Source/Developer/TraceAnalysis/Private/Store/AsioObject.h
@@ -0,0 +1,25 @@
+// Copyright Epic Games, Inc. All Rights Reserved.
+
+#pragma once
+
+#include "Asio/Asio.h"
+
+#if TRACE_WITH_ASIO
+
+namespace Trace
+{
+
+////////////////////////////////////////////////////////////////////////////////
+class FAsioObject
+{
+public:
+						FAsioObject(asio::io_context& InIoContext);
+	asio::io_context&	GetIoContext();
+
+private:
+	asio::io_context&	IoContext;
+};
+
+} // namespace Trace
+
+#endif // TRACE_WITH_ASIO
diff --git a/Engine/Source/Developer/TraceAnalysis/Private/Store/AsioRecorder.cpp b/Engine/Source/Developer/TraceAnalysis/Private/Store/AsioRecorder.cpp
new file mode 100644
--- /dev/null
+++ b/Engine/Source/Developer/TraceAnalysis/Private/Store/AsioRecorder.cpp
@@ -0,0 +1,118 @@
+// Copyright Epic Games, Inc. All Rights Reserved.
+
+#include "AsioRecorder.h"
+
+#if TRACE_WITH_ASIO
+
+#include "AsioIoable.h"
+#include "AsioSocket.h"
+#include "AsioStore.h"
+
+namespace Trace
+{
+
+////////////////////////////////////////////////////////////////////////////////
+class FAsioRecorderPeer
+	: public FAsioIoSink
+{
+public:
+						FAsioRecorderPeer(asio::ip::tcp::socket& Socket, FAsioWriteable* InOutput);
+	virtual				~FAsioRecorderPeer();
+	void				Close();
+
+private:
+	virtual void		OnIoComplete(uint32 Id, int32 Size) override;
+	static const uint32	BufferSize = 64 * 1024;
+	enum				{ OpStart, OpSocketRead, OpFileWrite };
+	FAsioSocket			Input;
+	FAsioWriteable*		Output;
+	uint8				Buffer[BufferSize];
+};
+
+////////////////////////////////////////////////////////////////////////////////
+FAsioRecorderPeer::FAsioRecorderPeer(asio::ip::tcp::socket& Socket, FAsioWriteable* InOutput)
+: Input(Socket)
+, Output(InOutput)
+{
+	OnIoComplete(OpStart, 0);
+}
+
+////////////////////////////////////////////////////////////////////////////////
+FAsioRecorderPeer::~FAsioRecorderPeer()
+{
+	delete Output;
+}
+
+////////////////////////////////////////////////////////////////////////////////
+void FAsioRecorderPeer::Close()
+{
+	Input.Close();
+	Output->Close();
+}
+
+////////////////////////////////////////////////////////////////////////////////
+void FAsioRecorderPeer::OnIoComplete(uint32 Id, int32 Size)
+{
+	if (Size < 0)
+	{
+		Close();
+		return;
+	}
+
+	switch (Id)
+	{
+	case OpSocketRead:
+		Output->Write(Buffer, Size, this, OpFileWrite);
+		break;
+
+	case OpStart:
+	case OpFileWrite:
+		Input.ReadSome(Buffer, BufferSize, this, OpSocketRead);
+		break;
+	}
+}
+
+
+
+////////////////////////////////////////////////////////////////////////////////
+FAsioRecorder::FAsioRecorder(asio::io_context& IoContext, FAsioStore& InStore)
+: FAsioTcpServer(IoContext)
+, FAsioTickable(IoContext)
+, Store(InStore)
+{
+	StartTick(500);
+}
+
+////////////////////////////////////////////////////////////////////////////////
+FAsioRecorder::~FAsioRecorder()
+{
+	for (FAsioRecorderPeer* Peer : Peers)
+	{
+		Peer->Close();
+		delete Peer;
+	}
+}
+
+////////////////////////////////////////////////////////////////////////////////
+bool FAsioRecorder::OnAccept(asio::ip::tcp::socket& Socket)
+{
+	FAsioWriteable* Output = Store.CreateTrace();
+	if (Output == nullptr)
+	{
+		return true;
+	}
+
+	FAsioRecorderPeer* Peer = new FAsioRecorderPeer(Socket, Output);
+	Peers.Add(Peer);
+	return true;
+}
+
+////////////////////////////////////////////////////////////////////////////////
+void FAsioRecorder::OnTick()
+{
+	/* cleanup dead clients */
+}
+
+} // namespace Trace
+
+#endif // TRACE_WITH_ASIO
diff --git a/Engine/Source/Developer/TraceAnalysis/Private/Store/AsioRecorder.h b/Engine/Source/Developer/TraceAnalysis/Private/Store/AsioRecorder.h
new file mode 100644
--- /dev/null
+++ b/Engine/Source/Developer/TraceAnalysis/Private/Store/AsioRecorder.h
@@ -0,0 +1,36 @@
+// Copyright Epic Games, Inc. All Rights Reserved.
+
+#pragma once
+
+#include "Asio/Asio.h"
+
+#if TRACE_WITH_ASIO
+
+#include "AsioTcpServer.h"
+#include "AsioTickable.h"
+
+namespace Trace
+{
+
+class FAsioStore;
+class FAsioRecorderPeer;
+
+////////////////////////////////////////////////////////////////////////////////
+class FAsioRecorder
+	: public FAsioTcpServer
+	, public FAsioTickable
+{
+public:
+								FAsioRecorder(asio::io_context& IoContext, FAsioStore& InStore);
+								~FAsioRecorder();
+
+private:
+	virtual bool				OnAccept(asio::ip::tcp::socket& Socket) override;
+	virtual void				OnTick() override;
+	TArray<FAsioRecorderPeer*>	Peers;
+	FAsioStore&					Store;
+};
+
+} // namespace Trace
+
+#endif // TRACE_WITH_ASIO
diff --git a/Engine/Source/Developer/TraceAnalysis/Private/Store/AsioSocket.cpp b/Engine/Source/Developer/TraceAnalysis/Private/Store/AsioSocket.cpp
new file mode 100644
--- /dev/null
+++ b/Engine/Source/Developer/TraceAnalysis/Private/Store/AsioSocket.cpp
@@ -0,0 +1,94 @@
+// Copyright Epic Games, Inc. All Rights Reserved.
+
+#include "AsioSocket.h"
+
+#if TRACE_WITH_ASIO
+
+namespace Trace
+{
+
+////////////////////////////////////////////////////////////////////////////////
+FAsioSocket::FAsioSocket(asio::ip::tcp::socket& InSocket)
+: Socket(MoveTemp(InSocket))
+{
+}
+
+////////////////////////////////////////////////////////////////////////////////
+FAsioSocket::~FAsioSocket()
+{
+}
+
+////////////////////////////////////////////////////////////////////////////////
+bool FAsioSocket::IsOpen() const
+{
+	return Socket.is_open();
+}
+
+////////////////////////////////////////////////////////////////////////////////
+void FAsioSocket::Close()
+{
+	Socket.close();
+}
+
+////////////////////////////////////////////////////////////////////////////////
+bool FAsioSocket::Read(void* Dest, uint32 Size, FAsioIoSink* Sink, uint32 Id)
+{
+	if (!SetSink(Sink, Id))
+	{
+		return false;
+	}
+
+	asio::async_read(
+		Socket,
+        asio::buffer(Dest, Size),
+		[this] (const asio::error_code& ErrorCode, size_t BytesReceived)
+		{
+			OnIoComplete(ErrorCode, BytesReceived);
+		}
+	);
+
+	return true;
+}
+
+////////////////////////////////////////////////////////////////////////////////
+bool FAsioSocket::ReadSome(void* Dest, uint32 BufferSize, FAsioIoSink* Sink, uint32 Id)
+{
+	if (!SetSink(Sink, Id))
+	{
+		return false;
+	}
+
+	Socket.async_receive(
+        asio::buffer(Dest, BufferSize),
+		[this] (const asio::error_code& ErrorCode, size_t BytesReceived)
+		{
+			return OnIoComplete(ErrorCode, BytesReceived);
+		}
+	);
+
+	return true;
+}
+
+////////////////////////////////////////////////////////////////////////////////
+bool FAsioSocket::Write(const void* Src, uint32 Size, FAsioIoSink* Sink, uint32 Id)
+{
+	if (!SetSink(Sink, Id))
+	{
+		return false;
+	}
+
+	asio::async_write(
+		Socket,
+		asio::buffer(Src, Size),
+		[this] (const asio::error_code& ErrorCode, size_t BytesSent)
+		{
+			return OnIoComplete(ErrorCode, BytesSent);
+		}
+	);
+
+	return true;
+}
+
+} // namespace Trace
+
+#endif // TRACE_WITH_ASIO
diff --git a/Engine/Source/Developer/TraceAnalysis/Private/Store/AsioSocket.h b/Engine/Source/Developer/TraceAnalysis/Private/Store/AsioSocket.h
new file mode 100644
--- /dev/null
+++ b/Engine/Source/Developer/TraceAnalysis/Private/Store/AsioSocket.h
@@ -0,0 +1,34 @@
+// Copyright Epic Games, Inc. All Rights Reserved.
+
+#pragma once
+
+#include "Asio/Asio.h"
+
+#if TRACE_WITH_ASIO
+
+#include "AsioIoable.h"
+
+namespace Trace
+{
+
+////////////////////////////////////////////////////////////////////////////////
+class FAsioSocket
+	: public FAsioReadable
+	, public FAsioWriteable
+{
+public:
+							FAsioSocket(asio::ip::tcp::socket& InSocket);
+	virtual					~FAsioSocket();
+	bool					IsOpen() const;
+	virtual void			Close() override;
+	virtual bool			Read(void* Dest, uint32 Size, FAsioIoSink* Sink, uint32 Id) override;
+	virtual bool			ReadSome(void* Dest, uint32 DestSize, FAsioIoSink* Sink, uint32 Id) override;
+	virtual bool			Write(const void* Src, uint32 Size, FAsioIoSink* Sink, uint32 Id) override;
+
+private:
+	asio::ip::tcp::socket	Socket;
+};
+
+} // namespace Trace
+
+#endif // TRACE_WITH_ASIO
diff --git a/Engine/Source/Developer/TraceAnalysis/Private/Store/AsioStore.cpp b/Engine/Source/Developer/TraceAnalysis/Private/Store/AsioStore.cpp
new file mode 100644
--- /dev/null
+++ b/Engine/Source/Developer/TraceAnalysis/Private/Store/AsioStore.cpp
@@ -0,0 +1,154 @@
+// Copyright Epic Games, Inc. All Rights Reserved.
+
+#include "AsioStore.h"
+
+#if TRACE_WITH_ASIO
+
+namespace Trace
+{
+
+////////////////////////////////////////////////////////////////////////////////
+FAsioStore::FAsioStore(asio::io_context& IoContext, const TCHAR* InStoreDir)
+: FAsioObject(IoContext)
+, StoreDir(InStoreDir)
+{
+	IPlatformFile& FileSystem = IPlatformFile::GetPlatformPhysical();
+	FileSystem.IterateDirectory(InStoreDir, [this] (const TCHAR* Name, bool IsDirectory)
+	{
+#if 0
+		if (IsDirectory)
+		{
+			int32 Id = FCString::Atoi(Name);
+			LastTraceId = (Id < LastTraceId) ? Id : LastTraceId;
+		}
+#else
+		const TCHAR* Dot = FCString::Strrchr(Name, '.');
+		if (Dot == nullptr || FCString::Strcmp(Dot, TEXT(".utrace")))
+		{
+			return true;
+		}
+
+		for (const TCHAR* c = Dot; c > Name; --c)
+		{
+			if (c[-1] == '\\' || c[-1] == '/')
+			{
+				Name = c;
+				break;
+			}
+		}
+
+		FTrace* Trace = new FTrace();
+		Trace->Name.AppendChars(Name, int32(Dot - Name));
+		Traces.Add(Trace);
+#endif // 0
+
+		return true;
+	});
+}
+
+////////////////////////////////////////////////////////////////////////////////
+FAsioStore::~FAsioStore()
+{
+	for (FTrace* Trace : Traces)
+	{
+		delete Trace;
+	}
+}
+
+////////////////////////////////////////////////////////////////////////////////
+uint32 FAsioStore::GetTraceCount() const
+{
+	return Traces.Num();
+}
+
+////////////////////////////////////////////////////////////////////////////////
+const FAsioStore::FTrace* FAsioStore::GetTraceInfo(uint32 Index) const
+{
+	if (Index >= uint32(Traces.Num()))
+	{
+		return nullptr;
+	}
+
+	return Traces[Index];
+}
+
+////////////////////////////////////////////////////////////////////////////////
+FAsioStore::FTrace* FAsioStore::GetTrace(uint32 Id)
+{
+	for (FTrace* Trace : Traces)
+	{
+		if (Trace->GetId() == Id)
+		{
+			return Trace;
+		}
+	}
+
+	return nullptr;
+}
+
+////////////////////////////////////////////////////////////////////////////////
+FAsioWriteable* FAsioStore::CreateTrace()
+{
+	IPlatformFile& PlatformFile = IPlatformFile::GetPlatformPhysical();
+
+	FString TracePath;
+#if 0
+	bool bOk = false;
+	for (int i = 0; i < 256; ++i)
+	{
+		uint32 TraceId = ++LastTraceId;
+
+		TracePath = StoreDir;
+		TracePath.Appendf(TEXT("/%05d"), TraceId);
+
+		if (!PlatformFile.DirectoryExists(*TracePath) && PlatformFile.CreateDirectory(*TracePath))
+		{
+			bOk = true;
+			break;
+		}
+	}
+
+	if (!bOk)
+	{
+		return nullptr;
+	}
+
+	TracePath += TEXT("/data.utrace");
+#else
+	FString Prefix = FDateTime::Now().ToString(TEXT("%Y%m%d_%H%M%S"));
+
+	TracePath = StoreDir;
+	TracePath.Appendf(TEXT("/%s.utrace"), *Prefix);
+	for (uint64 Index = 0; PlatformFile.FileExists(*TracePath); ++Index)
+	{
+		TracePath = StoreDir;
+		TracePath.Appendf(TEXT("/%s_%d.utrace"), *Prefix, Index);
+	}
+#endif // 0
+
+	return FAsioFile::WriteFile(GetIoContext(), *TracePath);
+}
+
+////////////////////////////////////////////////////////////////////////////////
+FAsioReadable* FAsioStore::OpenTrace(uint32 Id)
+{
+	FTrace* Trace = GetTrace(Id);
+	if (Trace == nullptr)
+	{
+		return nullptr;
+	}
+
+	FString TracePath;
+	TracePath = StoreDir;
+#if 0
+	TracePath.Appendf(TEXT("/%05d/data.utrace"), Id);
+#else
+	TracePath.Appendf(TEXT("/%s.utrace"), Trace->GetName());
+#endif // 0
+
+	return FAsioFile::ReadFile(GetIoContext(), *TracePath);
+}
+
+} // namespace Trace
+
+#endif // TRACE_WITH_ASIO
diff --git a/Engine/Source/Developer/TraceAnalysis/Private/Store/AsioStore.h b/Engine/Source/Developer/TraceAnalysis/Private/Store/AsioStore.h
new file mode 100644
--- /dev/null
+++ b/Engine/Source/Developer/TraceAnalysis/Private/Store/AsioStore.h
@@ -0,0 +1,61 @@
+// Copyright Epic Games, Inc. All Rights Reserved.
+
+#pragma once
+
+#include "Asio/Asio.h"
+
+#if TRACE_WITH_ASIO
+
+#include "AsioFile.h"
+#include "AsioObject.h"
+#include "Utils.h"
+
+namespace Trace
+{
+
+class FAsioReadable;
+class FAsioWriteable;
+
+////////////////////////////////////////////////////////////////////////////////
+enum class EStoreVersion
+{
+	Value = 0x0100, // 0xMMmm MM=major, mm=minor
+};
+
+////////////////////////////////////////////////////////////////////////////////
+class FAsioStore
+	: public FAsioObject
+{
+public:
+	class FTrace
+	{
+	public:
+		const TCHAR*	GetName() const		{ return *Name; }
+		uint32			GetId() const		{ return QuickStoreHash(GetName()); }
+		uint64			GetSize() const		{ return 3; }
+
+	private:
+		friend			FAsioStore;
+		FString			Name;
+	};
+
+						FAsioStore(asio::io_context& IoContext, const TCHAR* InStoreDir);
+						~FAsioStore();
+	uint32				GetTraceCount() const;
+	const FTrace*		GetTraceInfo(uint32 Index) const;
+	int32				ReadTrace(uint32 Id);
+	FAsioWriteable*		CreateTrace();
+	FAsioReadable*		OpenTrace(uint32 Id);
+
+private:
+	FTrace*				GetTrace(uint32 Id);
+	FString				StoreDir;
+	TArray<FTrace*>		Traces;
+#if 0
+	int32				LastTraceId = -1;
+#endif // 0
+};
+
+} // namespace Trace
+
+#endif // TRACE_WITH_ASIO
diff --git a/Engine/Source/Developer/TraceAnalysis/Private/Store/AsioStoreCborServer.cpp b/Engine/Source/Developer/TraceAnalysis/Private/Store/AsioStoreCborServer.cpp
new file mode 100644
--- /dev/null
+++ b/Engine/Source/Developer/TraceAnalysis/Private/Store/AsioStoreCborServer.cpp
@@ -0,0 +1,289 @@
+// Copyright Epic Games, Inc. All Rights Reserved.
+
+#include "AsioStoreCborServer.h"
+
+#if TRACE_WITH_ASIO
+
+#include "AsioIoable.h"
+#include "AsioObject.h"
+#include "AsioRecorder.h"
+#include "AsioSocket.h"
+#include "AsioStore.h"
+#include "AsioTraceRelay.h"
+#include "CborPayload.h"
+
+namespace Trace
+{
+
+////////////////////////////////////////////////////////////////////////////////
+class FAsioStoreCborPeer
+	: public FAsioIoSink
+	, public FAsioObject
+{
+public:
+					FAsioStoreCborPeer(asio::ip::tcp::socket& InSocket, FAsioStore& InStore, FAsioRecorder& InRecorder);
+	virtual			~FAsioStoreCborPeer();
+	bool			IsActive() const;
+	void			Close();
+
+protected:
+	enum class EStatusCode
+	{
+		Success				= 200,
+		BadRequest			= 400,
+		MethodNotAllowed	= 405,
+		InternalError		= 500,
+	};
+
+	void			OnPayload();
+	void			OnConnect();
+	void			OnStatus();
+	void			OnTraceCount();
+	void			OnTraceInfo();
+	void			OnTraceRead();
+	void			SendError(EStatusCode StatusCode);
+	void			SendResponse(const FPayload& Payload);
+	virtual void	OnIoComplete(uint32 Id, int32 Size) override;
+	enum			{ OpStart, OpReadPayloadSize, OpReadPayload, OpSendResponse };
+	FAsioStore&		Store;
+	FAsioRecorder&	Recorder;
+	FAsioSocket		Socket;
+	uint32			PayloadSize;
+	FResponse		Response;
+};
+
+////////////////////////////////////////////////////////////////////////////////
+FAsioStoreCborPeer::FAsioStoreCborPeer(
+	asio::ip::tcp::socket& InSocket,
+	FAsioStore& InStore,
+	FAsioRecorder& InRecorder)
+: FAsioObject(InSocket.get_executor().context())
+, Store(InStore)
+, Recorder(InRecorder)
+, Socket(InSocket)
+{
+	OnIoComplete(OpStart, 0);
+}
+
+////////////////////////////////////////////////////////////////////////////////
+FAsioStoreCborPeer::~FAsioStoreCborPeer()
+{
+	Close();
+}
+
+////////////////////////////////////////////////////////////////////////////////
+bool FAsioStoreCborPeer::IsActive() const
+{
+	return Socket.IsOpen();
+}
+
+////////////////////////////////////////////////////////////////////////////////
+void FAsioStoreCborPeer::Close()
+{
+	if (IsActive())
+	{
+		Socket.Close();
+	}
+}
+
+////////////////////////////////////////////////////////////////////////////////
+void FAsioStoreCborPeer::OnConnect()
+{
+	TPayloadBuilder<> Builder((int32)EStatusCode::Success);
+	Builder.AddParam("version", int32(EStoreVersion::Value));
+	SendResponse(Builder.Done());
+}
+
+////////////////////////////////////////////////////////////////////////////////
+void FAsioStoreCborPeer::OnStatus()
+{
+	TPayloadBuilder<> Builder((int32)EStatusCode::Success);
+	Builder.AddParam("recorder_port", Recorder.GetPort());
+	SendResponse(Builder.Done());
+}
+
+////////////////////////////////////////////////////////////////////////////////
+void FAsioStoreCborPeer::OnTraceCount()
+{
+	TPayloadBuilder<> Builder((int32)EStatusCode::Success);
+	Builder.AddParam("count", Store.GetTraceCount());
+	SendResponse(Builder.Done());
+}
+
+////////////////////////////////////////////////////////////////////////////////
+void FAsioStoreCborPeer::OnTraceInfo()
+{
+	int32 Index = int32(Response.GetInteger("index", -1));
+	if (Index < 0)
+	{
+		return SendError(EStatusCode::BadRequest);
+	}
+
+	const FAsioStore::FTrace* Trace = Store.GetTraceInfo(Index);
+	if (Trace == nullptr)
+	{
+		return SendError(EStatusCode::BadRequest);
+	}
+
+	TPayloadBuilder<> Builder((int32)EStatusCode::Success);
+	Builder.AddParam("id", Trace->GetId());
+	Builder.AddParam("size", Trace->GetSize());
+	SendResponse(Builder.Done());
+}
+
+////////////////////////////////////////////////////////////////////////////////
+void FAsioStoreCborPeer::OnTraceRead()
+{
+	uint32 Id = uint32(Response.GetInteger("id", 0));
+	if (!Id)
+	{
+		return SendError(EStatusCode::BadRequest);
+	}
+
+	FAsioReadable* Input = Store.OpenTrace(Id);
+	if (Input == nullptr)
+	{
+		return SendError(EStatusCode::BadRequest);
+	}
+
+	FAsioTraceRelay* Relay = new FAsioTraceRelay(GetIoContext(), Input);
+	uint32 Port = Relay->GetPort();
+	if (!Port)
+	{
+		delete Relay;
+		return SendError(EStatusCode::InternalError);
+	}
+
+	TPayloadBuilder<> Builder((int32)EStatusCode::Success);
+	Builder.AddParam("port", Port);
+	SendResponse(Builder.Done());
+}
+
+////////////////////////////////////////////////////////////////////////////////
+void FAsioStoreCborPeer::OnPayload()
+{
+	const char* Method = Response.GetString("$method");
+	if (Method[0] == '\0')
+	{
+		SendError(EStatusCode::BadRequest);
+		return;
+	}
+
+	struct {
+		uint32	Hash;
+		void	(FAsioStoreCborPeer::*Func)();
+	} const DispatchTable[] = {
+		{ QuickStoreHash("connect"),	&FAsioStoreCborPeer::OnConnect },
+		{ QuickStoreHash("status"),		&FAsioStoreCborPeer::OnStatus },
+		{ QuickStoreHash("trace_count"),&FAsioStoreCborPeer::OnTraceCount },
+		{ QuickStoreHash("trace_info"),	&FAsioStoreCborPeer::OnTraceInfo },
+		{ QuickStoreHash("trace_read"),	&FAsioStoreCborPeer::OnTraceRead },
+	};
+
+	uint32 MethodHash = QuickStoreHash(Method);
+	for (const auto& Row : DispatchTable)
+	{
+		if (Row.Hash == MethodHash)
+		{
+			return (this->*(Row.Func))();
+		}
+	}
+
+	SendError(EStatusCode::MethodNotAllowed);
+}
+
+////////////////////////////////////////////////////////////////////////////////
+void FAsioStoreCborPeer::SendError(EStatusCode StatusCode)
+{
+	TPayloadBuilder<> Builder((int32)StatusCode);
+	FPayload Payload = Builder.Done();
+	SendResponse(Payload);
+}
+
+////////////////////////////////////////////////////////////////////////////////
+void FAsioStoreCborPeer::SendResponse(const FPayload& Payload)
+{
+	struct {
+		uint32	Size;
+		uint8	Data[];
+	}* Dest;
+
+	Dest = decltype(Dest)(Response.Reserve(sizeof(*Dest) + Payload.Size));
+	Dest->Size = Payload.Size;
+	memcpy(Dest->Data, Payload.Data, Payload.Size);
+	Socket.Write(Dest, sizeof(*Dest) + Payload.Size, this, OpSendResponse);
+}
+
+////////////////////////////////////////////////////////////////////////////////
+void FAsioStoreCborPeer::OnIoComplete(uint32 Id, int32 Size)
+{
+	if (Size < 0)
+	{
+		Socket.Close();
+		return;
+	}
+
+	switch (Id)
+	{
+	case OpStart:
+	case OpSendResponse:
+		Socket.Read(&PayloadSize, sizeof(uint32), this, OpReadPayloadSize);
+		break;
+
+	case OpReadPayloadSize:
+		Socket.Read(Response.Reserve(PayloadSize), PayloadSize, this, OpReadPayload);
+		break;
+
+	case OpReadPayload:
+		OnPayload();
+		break;
+	}
+}
+
+
+
+////////////////////////////////////////////////////////////////////////////////
+FAsioStoreCborServer::FAsioStoreCborServer(
+	asio::io_context& IoContext,
+	FAsioStore& InStore,
+	FAsioRecorder& InRecorder)
+: FAsioTcpServer(IoContext)
+, Store(InStore)
+, Recorder(InRecorder)
+{
+	StartServer();
+}
+
+////////////////////////////////////////////////////////////////////////////////
+FAsioStoreCborServer::~FAsioStoreCborServer()
+{
+	for (FAsioStoreCborPeer* Peer : Peers)
+	{
+		Peer->Close();
+		delete Peer;
+	}
+}
+
+////////////////////////////////////////////////////////////////////////////////
+FAsioStore& FAsioStoreCborServer::GetStore() const
+{
+	return Store;
+}
+
+////////////////////////////////////////////////////////////////////////////////
+FAsioRecorder& FAsioStoreCborServer::GetRecorder() const
+{
+	return Recorder;
+}
+
+////////////////////////////////////////////////////////////////////////////////
+bool FAsioStoreCborServer::OnAccept(asio::ip::tcp::socket& Socket)
+{
+	FAsioStoreCborPeer* Peer = new FAsioStoreCborPeer(Socket, Store, Recorder);
+	Peers.Add(Peer);
+	return true;
+}
+
+} // namespace Trace
+
+#endif // TRACE_WITH_ASIO
diff --git a/Engine/Source/Developer/TraceAnalysis/Private/Store/AsioStoreCborServer.h b/Engine/Source/Developer/TraceAnalysis/Private/Store/AsioStoreCborServer.h
new file mode 100644
--- /dev/null
+++ b/Engine/Source/Developer/TraceAnalysis/Private/Store/AsioStoreCborServer.h
@@ -0,0 +1,37 @@
+// Copyright Epic Games, Inc. All Rights Reserved.
+
+#pragma once
+
+#include "Asio/Asio.h"
+
+#if TRACE_WITH_ASIO
+
+#include "AsioTcpServer.h"
+
+namespace Trace
+{
+
+class FAsioRecorder;
+class FAsioStore;
+class FAsioStoreCborPeer;
+
+////////////////////////////////////////////////////////////////////////////////
+class FAsioStoreCborServer
+	: public FAsioTcpServer
+{
+public:
+								FAsioStoreCborServer(asio::io_context& IoContext, FAsioStore& InStore, FAsioRecorder& InRecorder);
+								~FAsioStoreCborServer();
+	FAsioStore&					GetStore() const;
+	FAsioRecorder&				GetRecorder() const;
+
+private:
+	virtual bool				OnAccept(asio::ip::tcp::socket& Socket) override;
+	TArray<FAsioStoreCborPeer*>	Peers;
+	FAsioStore&					Store;
+	FAsioRecorder&				Recorder;
+};
+
+} // namespace Trace
+
+#endif // TRACE_WITH_ASIO
diff --git a/Engine/Source/Developer/TraceAnalysis/Private/Store/AsioTcpServer.cpp b/Engine/Source/Developer/TraceAnalysis/Private/Store/AsioTcpServer.cpp
new file mode 100644
--- /dev/null
+++ b/Engine/Source/Developer/TraceAnalysis/Private/Store/AsioTcpServer.cpp
@@ -0,0 +1,104 @@
+// Copyright Epic Games, Inc. All Rights Reserved.
+
+#include "AsioTcpServer.h"
+
+#if TRACE_WITH_ASIO
+
+namespace Trace
+{
+
+////////////////////////////////////////////////////////////////////////////////
+FAsioTcpServer::FAsioTcpServer(asio::io_context& IoContext)
+: FAsioObject(IoContext)
+, Acceptor(IoContext)
+{
+}
+
+////////////////////////////////////////////////////////////////////////////////
+FAsioTcpServer::~FAsioTcpServer()
+{
+}
+
+////////////////////////////////////////////////////////////////////////////////
+uint32 FAsioTcpServer::GetPort() const
+{
+	using asio::ip::tcp;
+
+	if (!Acceptor.is_open())
+	{
+		return 0;
+	}
+
+	asio::error_code ErrorCode;
+	tcp::endpoint Endpoint = Acceptor.local_endpoint(ErrorCode);
+	if (ErrorCode)
+	{
+		return 0;
+	}
+
+	return Endpoint.port();
+}
+
+////////////////////////////////////////////////////////////////////////////////
+bool FAsioTcpServer::StartServer(uint32 Port)
+{
+	if (Acceptor.is_open())
+	{
+		return false;
+	}
+
+	using asio::ip::tcp;
+
+    tcp::endpoint Endpoint(tcp::v4(), uint16(Port));
+	tcp::acceptor TempAcceptor(GetIoContext(), Endpoint, false);
+
+	asio::error_code ErrorCode;
+	tcp::endpoint LocalEndpoint = TempAcceptor.local_endpoint(ErrorCode);
+	if (ErrorCode || LocalEndpoint.port() == 0)
+	{
+		return false;
+	}
+
+	Acceptor = MoveTemp(TempAcceptor);
+	AsyncAccept();
+	return true;
+}
+
+////////////////////////////////////////////////////////////////////////////////
+bool FAsioTcpServer::StopServer()
+{
+	if (!Acceptor.is_open())
+	{
+		return false;
+	}
+
+	Acceptor.close();
+	return true;
+}
+
+////////////////////////////////////////////////////////////////////////////////
+void FAsioTcpServer::AsyncAccept()
+{
+	using asio::ip::tcp;
+
+	Acceptor.async_accept([this] (
+		const asio::error_code& ErrorCode,
+		tcp::socket Socket)
+	{
+		if (ErrorCode)
+		{
+			return;
+		}
+
+		if (!OnAccept(Socket))
+		{
+			return;
+		}
+
+		AsyncAccept();
+	});
+}
+
+} // namespace Trace
+
+#endif // TRACE_WITH_ASIO
diff --git a/Engine/Source/Developer/TraceAnalysis/Private/Store/AsioTcpServer.h b/Engine/Source/Developer/TraceAnalysis/Private/Store/AsioTcpServer.h
new file mode 100644
--- /dev/null
+++ b/Engine/Source/Developer/TraceAnalysis/Private/Store/AsioTcpServer.h
@@ -0,0 +1,33 @@
+// Copyright Epic Games, Inc. All Rights Reserved.
+
+#pragma once
+
+#include "Asio/Asio.h"
+
+#if TRACE_WITH_ASIO
+
+#include "AsioObject.h"
+
+namespace Trace
+{
+
+////////////////////////////////////////////////////////////////////////////////
+class FAsioTcpServer
+	: public FAsioObject
+{
+public:
+							FAsioTcpServer(asio::io_context& IoContext);
+	virtual					~FAsioTcpServer();
+	virtual bool			OnAccept(asio::ip::tcp::socket& Socket) = 0;
+	uint32					GetPort() const;
+	bool					StartServer(uint32 Port=0);
+	bool					StopServer();
+
+private:
+	void					AsyncAccept();
+	asio::ip::tcp::acceptor	Acceptor;
+};
+
+} // namespace Trace
+
+#endif // TRACE_WITH_ASIO
diff --git a/Engine/Source/Developer/TraceAnalysis/Private/Store/AsioTickable.cpp b/Engine/Source/Developer/TraceAnalysis/Private/Store/AsioTickable.cpp
new file mode 100644
--- /dev/null
+++ b/Engine/Source/Developer/TraceAnalysis/Private/Store/AsioTickable.cpp
@@ -0,0 +1,78 @@
+// Copyright Epic Games, Inc. All Rights Reserved.
+
+#include "AsioTickable.h"
+
+#if TRACE_WITH_ASIO
+
+namespace Trace
+{
+
+////////////////////////////////////////////////////////////////////////////////
+FAsioTickable::FAsioTickable(asio::io_context& IoContext)
+: FAsioObject(IoContext)
+, Timer(IoContext)
+{
+}
+
+////////////////////////////////////////////////////////////////////////////////
+FAsioTickable::~FAsioTickable()
+{
+	StopTick();
+}
+
+////////////////////////////////////////////////////////////////////////////////
+bool FAsioTickable::StartTick(uint32 InMillisecondRate)
+{
+	if (MillisecondRate)
+	{
+		return false;
+	}
+
+	MillisecondRate = InMillisecondRate;
+	if (MillisecondRate)
+	{
+		AsyncTick();
+		return true;
+	}
+
+	return false;
+}
+
+////////////////////////////////////////////////////////////////////////////////
+bool FAsioTickable::StopTick()
+{
+	if (!MillisecondRate)
+	{
+		return false;
+	}
+
+	MillisecondRate = 0;
+	Timer.cancel();
+	return true;
+}
+
+////////////////////////////////////////////////////////////////////////////////
+void FAsioTickable::AsyncTick()
+{
+	auto StdTime = std::chrono::milliseconds(MillisecondRate);
+	Timer.expires_after(StdTime);
+
+	Timer.async_wait([this] (const asio::error_code& ErrorCode)
+	{
+		if (ErrorCode)
+		{
+			return;
+		}
+
+		OnTick();
+
+		if (MillisecondRate)
+		{
+			AsyncTick();
+		}
+	});
+}
+
+} // namespace Trace
+
+#endif // TRACE_WITH_ASIO
diff --git a/Engine/Source/Developer/TraceAnalysis/Private/Store/AsioTickable.h b/Engine/Source/Developer/TraceAnalysis/Private/Store/AsioTickable.h
new file mode 100644
--- /dev/null
+++ b/Engine/Source/Developer/TraceAnalysis/Private/Store/AsioTickable.h
@@ -0,0 +1,35 @@
+// Copyright Epic Games, Inc. All Rights Reserved.
+
+#pragma once
+
+#include "Asio/Asio.h"
+
+#if TRACE_WITH_ASIO
+
+#include "AsioObject.h"
+
+#include <chrono>
+
+namespace Trace
+{
+
+////////////////////////////////////////////////////////////////////////////////
+class FAsioTickable
+	: public FAsioObject
+{
+public:
+						FAsioTickable(asio::io_context& IoContext);
+	virtual				~FAsioTickable();
+	bool				StartTick(uint32 MillisecondRate);
+	bool				StopTick();
+	virtual void		OnTick() = 0;
+
+private:
+	void				AsyncTick();
+	asio::steady_timer	Timer;
+	uint32				MillisecondRate = 0;
+};
+
+} // namespace Trace
+
+#endif // TRACE_WITH_ASIO
diff --git a/Engine/Source/Developer/TraceAnalysis/Private/Store/AsioTraceRelay.cpp b/Engine/Source/Developer/TraceAnalysis/Private/Store/AsioTraceRelay.cpp
new file mode 100644
--- /dev/null
+++ b/Engine/Source/Developer/TraceAnalysis/Private/Store/AsioTraceRelay.cpp
@@ -0,0 +1,74 @@
+// Copyright Epic Games, Inc. All Rights Reserved.
+
+#include "AsioTraceRelay.h"
+
+#if TRACE_WITH_ASIO
+
+#include "AsioSocket.h"
+
+namespace Trace
+{
+
+////////////////////////////////////////////////////////////////////////////////
+FAsioTraceRelay::FAsioTraceRelay(asio::io_context& IoContext, FAsioReadable* InInput)
+: FAsioTcpServer(IoContext)
+, Input(InInput)
+{
+	StartServer();
+}
+
+////////////////////////////////////////////////////////////////////////////////
+FAsioTraceRelay::~FAsioTraceRelay()
+{
+	if (Output != nullptr)
+	{
+		delete Output;
+	}
+	delete Input;
+}
+
+////////////////////////////////////////////////////////////////////////////////
+void FAsioTraceRelay::Close()
+{
+	if (Output != nullptr)
+	{
+		Output->Close();
+	}
+
+	Input->Close();
+}
+
+////////////////////////////////////////////////////////////////////////////////
+bool FAsioTraceRelay::OnAccept(asio::ip::tcp::socket& Socket)
+{
+	Output = new FAsioSocket(Socket);
+	OnIoComplete(OpStart, 0);
+	return false;
+}
+
+////////////////////////////////////////////////////////////////////////////////
+void FAsioTraceRelay::OnIoComplete(uint32 Id, int32 Size)
+{
+	if (Size < 0)
+	{
+		Output->Close();
+		Input->Close();
+		return;
+	}
+
+	switch (Id)
+	{
+	case OpStart:
+	case OpSend:
+		Input->Read(Buffer, BufferSize, this, OpRead);
+		break;
+
+	case OpRead:
+		Output->Write(Buffer, Size, this, OpSend);
+		break;
+	}
+}
+
+} // namespace Trace
+
+#endif // TRACE_WITH_ASIO
diff --git a/Engine/Source/Developer/TraceAnalysis/Private/Store/AsioTraceRelay.h b/Engine/Source/Developer/TraceAnalysis/Private/Store/AsioTraceRelay.h
new file mode 100644
--- /dev/null
+++ b/Engine/Source/Developer/TraceAnalysis/Private/Store/AsioTraceRelay.h
@@ -0,0 +1,37 @@
+// Copyright Epic Games, Inc. All Rights Reserved.
+
+#pragma once
+
+#include "Asio/Asio.h"
+
+#if TRACE_WITH_ASIO
+
+#include "AsioIoable.h"
+#include "AsioTcpServer.h"
+
+namespace Trace
+{
+
+////////////////////////////////////////////////////////////////////////////////
+class FAsioTraceRelay
+	: public FAsioTcpServer
+	, public FAsioIoSink
+{
+public:
+						FAsioTraceRelay(asio::io_context& IoContext, FAsioReadable* InInput);
+						~FAsioTraceRelay();
+	void				Close();
+
+private:
+	virtual bool		OnAccept(asio::ip::tcp::socket& Socket) override;
+	virtual void		OnIoComplete(uint32 Id, int32 Size) override;
+	enum				{ OpStart, OpRead, OpSend };
+	static const uint32	BufferSize = 64 << 10;
+	FAsioReadable*		Input;
+	FAsioWriteable*		Output = nullptr;
+	uint8				Buffer[BufferSize];
+};
+
+} // namespace Trace
+
+#endif // TRACE_WITH_ASIO
diff --git a/Engine/Source/Developer/TraceAnalysis/Private/Store/CborPayload.h b/Engine/Source/Developer/TraceAnalysis/Private/Store/CborPayload.h
new file mode 100644
--- /dev/null
+++ b/Engine/Source/Developer/TraceAnalysis/Private/Store/CborPayload.h
@@ -0,0 +1,226 @@
+// Copyright Epic Games, Inc. All Rights Reserved.
+
+#pragma once
+
+#include "Asio/Asio.h"
+
+#if TRACE_WITH_ASIO
+
+#include "CborReader.h"
+#include "CborWriter.h"
+#include "Serialization/MemoryArchive.h"
+#include "Serialization/MemoryReader.h"
+
+namespace Trace
+{
+
+////////////////////////////////////////////////////////////////////////////////
+template <int BufferSize>
+struct TInlineMemoryWriter
+	: public FMemoryArchive
+{
+					TInlineMemoryWriter();
+	virtual void	Serialize(void* Data, int64 Num) override;
+	uint32			Used = 0;
+	uint8			Buffer[BufferSize];
+};
+
+////////////////////////////////////////////////////////////////////////////////
+template <int BufferSize>
+inline TInlineMemoryWriter<BufferSize>::TInlineMemoryWriter()
+{
+	SetIsSaving(true);
+}
+
+////////////////////////////////////////////////////////////////////////////////
+template <int BufferSize>
+inline void TInlineMemoryWriter<BufferSize>::Serialize(void* Data, int64 Num)
+{
+	uint32 NextUsed = Used + uint32(Num);
+	if (NextUsed <= BufferSize)
+	{
+		memcpy(Buffer + Used, Data, Num);
+	}
+	Used = NextUsed;
+}
+
+
+
+////////////////////////////////////////////////////////////////////////////////
+struct FPayload
+{
+	const uint8*	Data;
+	uint32			Size;
+};
+
+
+
+////////////////////////////////////////////////////////////////////////////////
+template <int Size=128>
+class TPayloadBuilder
+{
+public:
+									TPayloadBuilder(int32 StatusCode);
+	template <int N>				TPayloadBuilder(const char (&MethodName)[N]);
+	template <int N> void			AddParam(const char (&Name)[N], int64 Value);
+	template <int N, int M> void	AddParam(const char (&Name)[N], const char (&Value)[M]);
+	FPayload						Done();
+
+private:
+	TInlineMemoryWriter<Size>		MemoryWriter;
+	FCborWriter						CborWriter = &MemoryWriter;
+};
+
+////////////////////////////////////////////////////////////////////////////////
+template <int Size>
+inline TPayloadBuilder<Size>::TPayloadBuilder(int32 StatusCode)
+{
+	CborWriter.WriteContainerStart(ECborCode::Map, -1);
+	AddParam("$status", StatusCode);
+}
+
+////////////////////////////////////////////////////////////////////////////////
+template <int Size>
+template <int N>
+inline TPayloadBuilder<Size>::TPayloadBuilder(const char (&MethodName)[N])
+{
+	CborWriter.WriteContainerStart(ECborCode::Map, -1);
+	AddParam("$method", MethodName);
+}
+
+////////////////////////////////////////////////////////////////////////////////
+template <int Size>
+template <int N>
+inline void TPayloadBuilder<Size>::AddParam(const char (&Name)[N], int64 Value)
+{
+	CborWriter.WriteValue(Name, N);
+	CborWriter.WriteValue(Value);
+}
+
+////////////////////////////////////////////////////////////////////////////////
+template <int Size>
+template <int N, int M>
+inline void TPayloadBuilder<Size>::AddParam(const char (&Name)[N], const char (&Value)[M])
+{
+	CborWriter.WriteValue(Name, N);
+	CborWriter.WriteValue(Value, M);
+}
+
+////////////////////////////////////////////////////////////////////////////////
+template <int Size>
+inline FPayload TPayloadBuilder<Size>::Done()
+{
+	CborWriter.WriteContainerEnd();
+	return { MemoryWriter.Buffer, MemoryWriter.Used };
+}
+
+
+
+////////////////////////////////////////////////////////////////////////////////
+class FResponse
+{
+public:
+	int64			GetInteger(const char* Key, int64 Default=0) const;
+	const char*		GetString(const char* Key, const char* Default="") const;
+	const uint8*	GetData() const;
+	uint32			GetSize() const;
+	uint8*			Reserve(uint32 Size);
+
+private:
+	template <typename Type, typename LambdaType>
+	Type			GetValue(const char* Key, Type Default, LambdaType&& Lambda) const;
+	TArray<uint8>	Buffer;
+};
+
+////////////////////////////////////////////////////////////////////////////////
+inline const uint8* FResponse::GetData() const
+{
+	return Buffer.GetData();
+};
+
+////////////////////////////////////////////////////////////////////////////////
+inline uint32 FResponse::GetSize() const
+{
+	return Buffer.Num();
+}
+
+////////////////////////////////////////////////////////////////////////////////
+inline uint8* FResponse::Reserve(uint32 Size)
+{
+	Buffer.SetNumUninitialized(Size);
+	return Buffer.GetData();
+}
+
+////////////////////////////////////////////////////////////////////////////////
+template <typename Type, typename LambdaType>
+inline Type FResponse::GetValue(const char* Key, Type Default, LambdaType&& Lambda) const
+{
+	FMemoryReader MemoryReader(Buffer);
+	FCborReader CborReader(&MemoryReader);
+	FCborContext Context;
+
+	if (!CborReader.ReadNext(Context) || Context.MajorType() != ECborCode::Map)
+	{
+		return Default;
+	}
+
+	while (true)
+	{
+		// Read key
+		if (!CborReader.ReadNext(Context) || !Context.IsString())
+		{
+			return Default;
+		}
+
+		const char* String = Context.AsCString();
+		bool bIsTarget = (FCStringAnsi::Strcmp(Key, String) == 0);
+
+		// Read value
+		if (!CborReader.ReadNext(Context))
+		{
+			return Default;
+		}
+
+		if (bIsTarget)
+		{
+			return Lambda(Context, uint32(MemoryReader.Tell()));
+		}
+
+		if (String[0] == '$') __debugbreak();
+	}
+}
+
+////////////////////////////////////////////////////////////////////////////////
+inline int64 FResponse::GetInteger(const char* Key, int64 Default) const
+{
+	return GetValue(
+		Key,
+		Default,
+		[this] (const FCborContext& Context, int64)
+		{
+			return int32(Context.AsInt());
+		}
+	);
+}
+
+////////////////////////////////////////////////////////////////////////////////
+inline const char* FResponse::GetString(const char* Key, const char* Default) const
+{
+	return GetValue(
+		Key,
+		Default,
+		[this, Default] (const FCborContext& Context, uint32 Offset)
+		{
+			if (Context.IsString())
+			{
+				return (const char*)(Buffer.GetData() + Offset - Context.AsLength());
+			}
+
+			return Default;
+		}
+	);
+}
+
+} // namespace Trace
+
+#endif // TRACE_WITH_ASIO
diff --git a/Engine/Source/Developer/TraceAnalysis/Private/Store/StoreClient.cpp b/Engine/Source/Developer/TraceAnalysis/Private/Store/StoreClient.cpp
new file mode 100644
--- /dev/null
+++ b/Engine/Source/Developer/TraceAnalysis/Private/Store/StoreClient.cpp
@@ -0,0 +1,362 @@
+// Copyright Epic Games, Inc. All Rights Reserved.
+
+#include "Trace/StoreClient.h"
+#include "Asio/Asio.h"
+
+#if TRACE_WITH_ASIO
+
+#include "AsioStore.h"
+#include "CborPayload.h"
+
+namespace Trace
+{
+
+////////////////////////////////////////////////////////////////////////////////
+class FStoreCborClient
+{
+public:
+							FStoreCborClient();
+							~FStoreCborClient();
+	const FResponse&		GetResponse() const;
+	bool					Connect(const TCHAR* Host, uint16 Port);
+	bool					GetStatus();
+	bool					GetTraceCount();
+	bool					GetTraceInfo(uint32 Index);
+	asio::ip::tcp::socket*	ReadTrace(uint32 Id);
+
+private:
+	bool					Communicate(const FPayload& Payload);
+	asio::io_context		IoContext;
+	asio::ip::tcp::socket	Socket;
+	FResponse				Response;
+};
+
+////////////////////////////////////////////////////////////////////////////////
+FStoreCborClient::FStoreCborClient()
+: Socket(IoContext)
+{
+}
+
+////////////////////////////////////////////////////////////////////////////////
+FStoreCborClient::~FStoreCborClient()
+{
+}
+
+////////////////////////////////////////////////////////////////////////////////
+const FResponse& FStoreCborClient::GetResponse() const
+{
+	return Response;
+}
+
+////////////////////////////////////////////////////////////////////////////////
+bool FStoreCborClient::Connect(const TCHAR* Host, uint16 Port)
+{
+	FTCHARToUTF8 HostUtf8(Host);
+	char PortString[8];
+	FCStringAnsi::Sprintf(PortString, "%d", Port);
+
+	asio::ip::tcp::resolver Resolver(IoContext);
+	asio::ip::tcp::resolver::results_type Endpoints = Resolver.resolve(HostUtf8.Get(), PortString);
+
+	asio::error_code ErrorCode;
+
+	asio::connect(Socket, Endpoints, ErrorCode);
+	if (ErrorCode)
+	{
+		return false;
+	}
+
+	TPayloadBuilder<> Builder("connect");
+	Builder.AddParam("version", int32(EStoreVersion::Value));
+	FPayload Payload = Builder.Done();
+	return Communicate(Payload);
+}
+
+////////////////////////////////////////////////////////////////////////////////
+bool FStoreCborClient::Communicate(const FPayload& Payload)
+{
+	if (!Socket.is_open())
+	{
+		return false;
+	}
+
+	asio::error_code ErrorCode;
+
+	// Send the payload
+	uint32 PayloadSize = Payload.Size;
+	asio::write(Socket, asio::buffer(&PayloadSize, sizeof(PayloadSize)), ErrorCode);
+	asio::write(Socket, asio::buffer(Payload.Data, Payload.Size), ErrorCode);
+	if (ErrorCode)
+	{
+		Socket.close();
+		return false;
+	}
+
+	// Wait for a response
+	uint32 ResponseSize = 0;
+	asio::read(Socket, asio::buffer(&ResponseSize, sizeof(ResponseSize)), ErrorCode);
+	if (ErrorCode)
+	{
+		Socket.close();
+		return false;
+	}
+
+	if (ResponseSize == 0)
+	{
+		Socket.close();
+		return false;
+	}
+
+	uint8* Dest = Response.Reserve(ResponseSize);
+	asio::read(Socket, asio::buffer(Dest, ResponseSize), ErrorCode);
+	if (ErrorCode)
+	{
+		Socket.close();
+		return false;
+	}
+
+	return true;
+}
+
+////////////////////////////////////////////////////////////////////////////////
+bool FStoreCborClient::GetStatus()
+{
+	TPayloadBuilder<32> Builder("status");
+	FPayload Payload = Builder.Done();
+	return Communicate(Payload);
+}
+
+////////////////////////////////////////////////////////////////////////////////
+bool FStoreCborClient::GetTraceCount()
+{
+	TPayloadBuilder<32> Builder("trace_count");
+	FPayload Payload = Builder.Done();
+	return Communicate(Payload);
+}
+
+////////////////////////////////////////////////////////////////////////////////
+bool FStoreCborClient::GetTraceInfo(uint32 Index)
+{
+	TPayloadBuilder<> Builder("trace_info");
+	Builder.AddParam("index", Index);
+	FPayload Payload = Builder.Done();
+	return Communicate(Payload);
+}
+
+////////////////////////////////////////////////////////////////////////////////
+asio::ip::tcp::socket* FStoreCborClient::ReadTrace(uint32 Id)
+{
+	TPayloadBuilder<> Builder("trace_read");
+	Builder.AddParam("id", Id);
+	FPayload Payload = Builder.Done();
+	if (!Communicate(Payload))
+	{
+		return nullptr;
+	}
+
+	uint32 SenderPort = Response.GetInteger("port", 0);
+	if (!SenderPort)
+	{
+		return nullptr;
+	}
+
+	asio::ip::address ServerAddr = Socket.local_endpoint().address();
+	asio::ip::tcp::endpoint Endpoint(ServerAddr, uint16(SenderPort));
+
+	asio::error_code ErrorCode;
+	asio::ip::tcp::socket SenderSocket(IoContext);
+	SenderSocket.connect(Endpoint, ErrorCode);
+	if (ErrorCode)
+	{
+		return nullptr;
+	}
+
+	return new asio::ip::tcp::socket(MoveTemp(SenderSocket));
+}
+
+} // namespace Trace
+
+#endif // TRACE_WITH_ASIO
+
+
+
+#if !TRACE_WITH_ASIO
+struct FResponse
+{
+	int64		GetInteger(const char*, int64 Default) const		{ return Default; }
+	const char* GetString(const char*, const char* Default) const	{ return Default; }
+};
+#endif // !TRACE_WITH_ASIO
+
+
+
+namespace Trace
+{
+
+////////////////////////////////////////////////////////////////////////////////
+uint32 FStoreClient::FStatus::GetRecorderPort() const
+{
+	const auto* Response = (const FResponse*)this;
+	return Response->GetInteger("recorder_port", 0);
+}
+
+
+
+////////////////////////////////////////////////////////////////////////////////
+uint32 FStoreClient::FTraceInfo::GetId() const
+{
+	const auto* Response = (const FResponse*)this;
+	return Response->GetInteger("id", 0);
+}
+
+////////////////////////////////////////////////////////////////////////////////
+uint64 FStoreClient::FTraceInfo::GetSize() const
+{
+	const auto* Response = (const FResponse*)this;
+	return Response->GetInteger("size", 0);
+}
+
+
+
+////////////////////////////////////////////////////////////////////////////////
+FStoreClient::FTraceData::~FTraceData()
+{
+#if TRACE_WITH_ASIO
+	if (auto* Inner = (asio::ip::tcp::socket*)Handle)
+	{
+		Inner->close();
+		delete Inner;
+	}
+#endif // TRACE_WITH_ASIO
+}
+
+////////////////////////////////////////////////////////////////////////////////
+int32 FStoreClient::FTraceData::Read(void* Dest, uint32 DestSize)
+{
+#if TRACE_WITH_ASIO
+	auto* Inner = (asio::ip::tcp::socket*)Handle;
+	if (Inner == nullptr)
+	{
+		return -1;
+	}
+
+	asio::error_code ErrorCode;
+	size_t BytesRead = Inner->read_some(asio::buffer(Dest, DestSize), ErrorCode);
+	if (ErrorCode)
+	{
+		Inner->close();
+		return -1;
+	}
+
+	return int32(BytesRead);
+#else
+	return -1;
+#endif // TRACE_WITH_ASIO
+}
+
+////////////////////////////////////////////////////////////////////////////////
+bool FStoreClient::FTraceData::IsValid() const
+{
+	return (Handle != 0);
+}
+
+
+
+////////////////////////////////////////////////////////////////////////////////
+struct FStoreClient::FImpl
+#if TRACE_WITH_ASIO
+	: public FStoreCborClient
+#endif
+{
+};
+
+////////////////////////////////////////////////////////////////////////////////
+FStoreClient::FStoreClient(const TCHAR* Host, uint16 Port)
+: Impl(nullptr)
+{
+	Impl = new FStoreClient::FImpl();
+#if TRACE_WITH_ASIO
+	if (!Impl->Connect(Host, Port))
+#endif
+	{
+		delete Impl;
+		Impl = nullptr;
+	}
+}
+
+////////////////////////////////////////////////////////////////////////////////
+FStoreClient::~FStoreClient()
+{
+	delete Impl;
+}
+
+////////////////////////////////////////////////////////////////////////////////
+bool FStoreClient::IsValid() const
+{
+	return Impl != nullptr;
+}
+
+////////////////////////////////////////////////////////////////////////////////
+const FStoreClient::FStatus* FStoreClient::GetStatus()
+{
+#if TRACE_WITH_ASIO
+	if (Impl == nullptr || !Impl->GetStatus())
+	{
+		return nullptr;
+	}
+
+	const FResponse& Response = Impl->GetResponse();
+	return (FStatus*)(&Response);
+#else
+	return nullptr;
+#endif // TRACE_WITH_ASIO
+}
+
+////////////////////////////////////////////////////////////////////////////////
+uint32 FStoreClient::GetTraceCount()
+{
+#if TRACE_WITH_ASIO
+	if (Impl == nullptr || !Impl->GetTraceCount())
+	{
+		return 0;
+	}
+
+	return Impl->GetResponse().GetInteger("count", 0);
+#else
+	return 0;
+#endif // TRACE_WITH_ASIO
+}
+
+////////////////////////////////////////////////////////////////////////////////
+const FStoreClient::FTraceInfo* FStoreClient::GetTraceInfo(uint32 Index)
+{
+#if TRACE_WITH_ASIO
+	if (Impl == nullptr || !Impl->GetTraceInfo(Index))
+	{
+		return nullptr;
+	}
+
+	const FResponse& Response = Impl->GetResponse();
+	return (FTraceInfo*)(&Response);
+#else
+	return nullptr;
+#endif // TRACE_WITH_ASIO
+}
+
+////////////////////////////////////////////////////////////////////////////////
+bool FStoreClient::ReadTrace(uint32 Id, FTraceData& Out)
+{
+#if TRACE_WITH_ASIO
+	if (Impl == nullptr)
+	{
+		return false;
+	}
+
+	Out.Handle = UPTRINT(Impl->ReadTrace(Id));
+	return (Out.Handle != 0);
+#else
+	return false;
+#endif // TRACE_WITH_ASIO
+}
+
+} // namespace Trace
diff --git a/Engine/Source/Developer/TraceAnalysis/Private/Store/StoreService.cpp b/Engine/Source/Developer/TraceAnalysis/Private/Store/StoreService.cpp
new file mode 100644
--- /dev/null
+++ b/Engine/Source/Developer/TraceAnalysis/Private/Store/StoreService.cpp
@@ -0,0 +1,94 @@
+// Copyright Epic Games, Inc. All Rights Reserved.
+
+#include "Trace/StoreService.h"
+#include "Asio/Asio.h"
+
+#if TRACE_WITH_ASIO
+
+#include "AsioContext.h"
+#include "AsioRecorder.h"
+#include "AsioStore.h"
+#include "AsioStoreCborServer.h"
+
+namespace Trace
+{
+
+////////////////////////////////////////////////////////////////////////////////
+struct FStoreService::FImpl
+{
+public:
+							FImpl(const FDesc& Desc);
+	FAsioContext			Context;
+	FAsioStore				Store;
+	FAsioRecorder			Recorder;
+	FAsioStoreCborServer	CborServer;
+};
+
+////////////////////////////////////////////////////////////////////////////////
+FStoreService::FImpl::FImpl(const FDesc& Desc)
+: Context(Desc.ThreadCount)
+, Store(Context.Get(), Desc.StoreDir)
+, Recorder(Context.Get(), Store)
+, CborServer(Context.Get(), Store, Recorder)
+{
+}
+
+} // namespace Trace
+
+#endif // TRACE_WITH_ASIO
+
+
+
+namespace Trace
+{
+
+////////////////////////////////////////////////////////////////////////////////
+FStoreService::~FStoreService()
+{
+#if TRACE_WITH_ASIO
+	delete Impl;
+#endif
+}
+
+////////////////////////////////////////////////////////////////////////////////
+FStoreService* FStoreService::Create(const FDesc& InDesc)
+{
+#if TRACE_WITH_ASIO
+	FDesc Desc = InDesc;
+
+	IPlatformFile& PlatformFile = IPlatformFile::GetPlatformPhysical();
+	PlatformFile.CreateDirectory(Desc.StoreDir);
+
+	if (Desc.ThreadCount <= 0)
+	{
+		Desc.ThreadCount = FPlatformMisc::NumberOfCoresIncludingHyperthreads();
+	}
+
+	FStoreService::FImpl* Impl = new FStoreService::FImpl(Desc);
+	if (Desc.RecorderPort >= 0)
+	{
+		FAsioRecorder& Recorder = Impl->Recorder;
+		Recorder.StartServer(Desc.RecorderPort);
+	}
+
+	Impl->Context.Start();
+
+	FStoreService* Server = new FStoreService();
+	Server->Impl = Impl;
+	return Server;
+#else
+	return nullptr;
+#endif // TRACE_WITH_ASIO
+}
+
+////////////////////////////////////////////////////////////////////////////////
+uint32 FStoreService::GetPort() const
+{
+#if TRACE_WITH_ASIO
+	return Impl->CborServer.GetPort();
+#else
+	return 0;
+#endif // TRACE_WITH_ASIO
+}
+
+} // namespace Trace
diff --git a/Engine/Source/Developer/TraceAnalysis/Private/Store/Utils.h b/Engine/Source/Developer/TraceAnalysis/Private/Store/Utils.h
new file mode 100644
--- /dev/null
+++ b/Engine/Source/Developer/TraceAnalysis/Private/Store/Utils.h
@@ -0,0 +1,21 @@
+// Copyright Epic Games, Inc. All Rights Reserved.
+
+#pragma once
+
+namespace Trace
+{
+
+////////////////////////////////////////////////////////////////////////////////
+template <typename T>
+inline uint32 QuickStoreHash(const T* String)
+{
+	uint32 Value = 5381;
+	for (; *String; ++String)
+	{
+		Value = ((Value << 5) + Value) + *String;
+	}
+	return Value;
+}
+
+
+} // namespace Trace
diff --git a/Engine/Source/Developer/TraceAnalysis/Public/Trace/Store.h b/Engine/Source/Developer/TraceAnalysis/Public/Trace/Store.h
--- a/Engine/Source/Developer/TraceAnalysis/Public/Trace/Store.h
+++ b/Engine/Source/Developer/TraceAnalysis/Public/Trace/Store.h
@@ -61,4 +61,4 @@ public:
  */
 TRACEANALYSIS_API TSharedPtr<IStore> Store_Create(const TCHAR* StoreDir);
 
-}
+} // namespace Trace
diff --git a/Engine/Source/Developer/TraceAnalysis/Public/Trace/StoreClient.h b/Engine/Source/Developer/TraceAnalysis/Public/Trace/StoreClient.h
new file mode 100644
--- /dev/null
+++ b/Engine/Source/Developer/TraceAnalysis/Public/Trace/StoreClient.h
@@ -0,0 +1,84 @@
+// Copyright Epic Games, Inc. All Rights Reserved.
+
+#pragma once
+
+#include "HAL/Platform.h"
+
+namespace Trace
+{
+
+////////////////////////////////////////////////////////////////////////////////
+class TRACEANALYSIS_API FStoreClient
+{
+public:
+	struct TRACEANALYSIS_API FStatus
+	{
+		uint32			GetRecorderPort() const;
+	};
+
+	struct TRACEANALYSIS_API FTraceInfo
+	{
+		uint32			GetId() const;
+		uint64			GetSize() const;
+		//const TCHAR*	GetMetadata(const TCHAR* Key) const;
+		//template <typename Lambda> uint32 ReadMetadata(Lambda&& Callback) const;
+	};
+
+	struct TRACEANALYSIS_API FTraceData
+	{
+						~FTraceData();
+		int32			Read(void* Dest, uint32 DestSize);
+		bool			IsValid() const;
+
+	private:
+		friend			FStoreClient;
+		UPTRINT			Handle = 0;
+	};
+						FStoreClient(const TCHAR* Host, uint16 Port);
+						~FStoreClient();
+	bool				IsValid() const;
+	const FStatus*		GetStatus();
+	uint32				GetTraceCount();
+	const FTraceInfo*	GetTraceInfo(uint32 Index);
+	bool				ReadTrace(uint32 Id, FTraceData& Out);
+
+#if 0
+	template <typename Lambda> uint32 GetTraceInfos(uint32 StartIndex, uint32 Count, Lambda&& Callback) const;
+#endif // 0
+
+#if 0
+	struct TRACEANALYSIS_API FSessionInfo
+	{
+		uint32			GetId() const;
+		uint32			GetTraceId() const;
+		uint32			GetIpAddress() const;
+	};
+	uint32				GetSessionCount() const;
+	const FSessionInfo* GetSessionInfo(uint32 Index, Lambda&& Callback) const;
+#endif // 0
+#if 0
+	template <typename Lambda> uint32	GetSessionInfos(uint32 StartIndex, uint32 Count, Lambda&& Callback) const;
+#endif // 0
+
+#if 0
+	// -------
+	class IStoreSubscriber
+	{
+		virtual void	OnStoreEvent() = 0;
+	};
+	bool				Subscribe(IStoreSubscriber* Subscriber);
+	bool				Unsubscribe(IStoreSubscriber* Subscriber);
+#endif // 0
+
+private:
+	struct				FImpl;
+	FImpl*				Impl;
+
+private:
+						FStoreClient(const FStoreClient&) = delete;
+						FStoreClient(const FStoreClient&&) = delete;
+	void				operator = (const FStoreClient&) = delete;
+	void				operator = (const FStoreClient&&) = delete;
+};
+
+} // namespace Trace
diff --git a/Engine/Source/Developer/TraceAnalysis/Public/Trace/StoreService.h b/Engine/Source/Developer/TraceAnalysis/Public/Trace/StoreService.h
new file mode 100644
--- /dev/null
+++ b/Engine/Source/Developer/TraceAnalysis/Public/Trace/StoreService.h
@@ -0,0 +1,37 @@
+// Copyright Epic Games, Inc. All Rights Reserved.
+
+#pragma once
+
+#include "HAL/Platform.h"
+
+namespace Trace
+{
+
+////////////////////////////////////////////////////////////////////////////////
+class TRACEANALYSIS_API FStoreService
+{
+public:
+	struct FDesc
+	{
+		const TCHAR*		StoreDir;
+		int32				RecorderPort = 0; // 0:auto-assign, -1:off
+		int32				ThreadCount  = 0; // <=0:logical CPU count
+	};
+
+							~FStoreService();
+	static FStoreService*	Create(const FDesc& Desc);
+	uint32					GetPort() const;
+
+private:
+							FStoreService() = default;
+	struct					FImpl;
+	FImpl*					Impl;
+
+private:
+							FStoreService(const FStoreService&) = delete;
+							FStoreService(const FStoreService&&) = delete;
+	void					operator = (const FStoreService&) = delete;
+	void					operator = (const FStoreService&&) = delete;
+};
+
+} // namespace Trace
diff --git a/Engine/Source/Developer/TraceAnalysis/TraceAnalysis.Build.cs b/Engine/Source/Developer/TraceAnalysis/TraceAnalysis.Build.cs
--- a/Engine/Source/Developer/TraceAnalysis/TraceAnalysis.Build.cs
+++ b/Engine/Source/Developer/TraceAnalysis/TraceAnalysis.Build.cs
@@ -6,9 +6,16 @@ public class TraceAnalysis : ModuleRules
 {
 	public TraceAnalysis(ReadOnlyTargetRules Target) : base(Target)
 	{
+		PrivateDependencyModuleNames.Add("Cbor");
 		PrivateDependencyModuleNames.Add("Core");
 		PrivateDependencyModuleNames.Add("DirectoryWatcher");
 		PrivateDependencyModuleNames.Add("Sockets");
 		PrivateDependencyModuleNames.Add("TraceLog");
+
+		PrivateDefinitions.Add("ASIO_SEPARATE_COMPILATION");
+		PrivateDefinitions.Add("ASIO_STANDALONE");
+		PrivateDefinitions.Add("ASIO_NO_EXCEPTIONS");
+		PrivateDefinitions.Add("ASIO_NO_TYPEID");
+		PrivateIncludePaths.Add("Developer/TraceAnalysis/Private/Asio");
 	}
 }
