Removed debug code that was submitted in error

#rnx

diff --git a/Engine/Source/Developer/TraceAnalysis/Private/Store/CborPayload.h b/Engine/Source/Developer/TraceAnalysis/Private/Store/CborPayload.h
--- a/Engine/Source/Developer/TraceAnalysis/Private/Store/CborPayload.h
+++ b/Engine/Source/Developer/TraceAnalysis/Private/Store/CborPayload.h
@@ -185,8 +185,6 @@ inline Type FResponse::GetValue(const char* Key, Type Default, LambdaType&& Lamb
 		{
 			return Lambda(Context, uint32(MemoryReader.Tell()));
 		}
-
-		if (String[0] == '$') __debugbreak();
 	}
 }
 
