Addressed "-disableunity -nopch" matters.

#rnx

diff --git a/Engine/Source/Developer/TraceAnalysis/Private/Asio/Asio.h b/Engine/Source/Developer/TraceAnalysis/Private/Asio/Asio.h
--- a/Engine/Source/Developer/TraceAnalysis/Private/Asio/Asio.h
+++ b/Engine/Source/Developer/TraceAnalysis/Private/Asio/Asio.h
@@ -8,6 +8,8 @@
 
 #if TRACE_WITH_ASIO
 
+#include "HAL/Platform.h"
+
 #if PLATFORM_WINDOWS
 #	include "Windows/AllowWindowsPlatformTypes.h"
 #	include "Windows/AllowWindowsPlatformAtomics.h"
diff --git a/Engine/Source/Developer/TraceAnalysis/Private/Store/AsioContext.h b/Engine/Source/Developer/TraceAnalysis/Private/Store/AsioContext.h
--- a/Engine/Source/Developer/TraceAnalysis/Private/Store/AsioContext.h
+++ b/Engine/Source/Developer/TraceAnalysis/Private/Store/AsioContext.h
@@ -3,6 +3,7 @@
 #pragma once
 
 #include "Asio/Asio.h"
+#include "Containers/Array.h"
 
 #if TRACE_WITH_ASIO
 
diff --git a/Engine/Source/Developer/TraceAnalysis/Private/Store/AsioRecorder.h b/Engine/Source/Developer/TraceAnalysis/Private/Store/AsioRecorder.h
--- a/Engine/Source/Developer/TraceAnalysis/Private/Store/AsioRecorder.h
+++ b/Engine/Source/Developer/TraceAnalysis/Private/Store/AsioRecorder.h
@@ -8,6 +8,7 @@
 
 #include "AsioTcpServer.h"
 #include "AsioTickable.h"
+#include "Containers/Array.h"
 
 namespace Trace
 {
diff --git a/Engine/Source/Developer/TraceAnalysis/Private/Store/AsioSocket.cpp b/Engine/Source/Developer/TraceAnalysis/Private/Store/AsioSocket.cpp
--- a/Engine/Source/Developer/TraceAnalysis/Private/Store/AsioSocket.cpp
+++ b/Engine/Source/Developer/TraceAnalysis/Private/Store/AsioSocket.cpp
@@ -4,6 +4,8 @@
 
 #if TRACE_WITH_ASIO
 
+#include "Templates/UnrealTemplate.h"
+
 namespace Trace
 {
 
diff --git a/Engine/Source/Developer/TraceAnalysis/Private/Store/AsioStore.cpp b/Engine/Source/Developer/TraceAnalysis/Private/Store/AsioStore.cpp
--- a/Engine/Source/Developer/TraceAnalysis/Private/Store/AsioStore.cpp
+++ b/Engine/Source/Developer/TraceAnalysis/Private/Store/AsioStore.cpp
@@ -4,6 +4,11 @@
 
 #if TRACE_WITH_ASIO
 
+#include "HAL/PlatformFile.h"
+#include "Containers/UnrealString.h"
+#include "Misc/CString.h"
+#include "Misc/DateTime.h"
+
 namespace Trace
 {
 
diff --git a/Engine/Source/Developer/TraceAnalysis/Private/Store/AsioStore.h b/Engine/Source/Developer/TraceAnalysis/Private/Store/AsioStore.h
--- a/Engine/Source/Developer/TraceAnalysis/Private/Store/AsioStore.h
+++ b/Engine/Source/Developer/TraceAnalysis/Private/Store/AsioStore.h
@@ -8,6 +8,8 @@
 
 #include "AsioFile.h"
 #include "AsioObject.h"
+#include "Containers/Array.h"
+#include "Containers/UnrealString.h"
 #include "Utils.h"
 
 namespace Trace
diff --git a/Engine/Source/Developer/TraceAnalysis/Private/Store/AsioStoreCborServer.h b/Engine/Source/Developer/TraceAnalysis/Private/Store/AsioStoreCborServer.h
--- a/Engine/Source/Developer/TraceAnalysis/Private/Store/AsioStoreCborServer.h
+++ b/Engine/Source/Developer/TraceAnalysis/Private/Store/AsioStoreCborServer.h
@@ -7,6 +7,7 @@
 #if TRACE_WITH_ASIO
 
 #include "AsioTcpServer.h"
+#include "Containers/Array.h"
 
 namespace Trace
 {
diff --git a/Engine/Source/Developer/TraceAnalysis/Private/Store/AsioTcpServer.cpp b/Engine/Source/Developer/TraceAnalysis/Private/Store/AsioTcpServer.cpp
--- a/Engine/Source/Developer/TraceAnalysis/Private/Store/AsioTcpServer.cpp
+++ b/Engine/Source/Developer/TraceAnalysis/Private/Store/AsioTcpServer.cpp
@@ -4,6 +4,8 @@
 
 #if TRACE_WITH_ASIO
 
+#include "Templates/UnrealTemplate.h"
+
 namespace Trace
 {
 
diff --git a/Engine/Source/Developer/TraceAnalysis/Private/Store/StoreService.cpp b/Engine/Source/Developer/TraceAnalysis/Private/Store/StoreService.cpp
--- a/Engine/Source/Developer/TraceAnalysis/Private/Store/StoreService.cpp
+++ b/Engine/Source/Developer/TraceAnalysis/Private/Store/StoreService.cpp
@@ -9,6 +9,7 @@
 #include "AsioRecorder.h"
 #include "AsioStore.h"
 #include "AsioStoreCborServer.h"
+#include "HAL/PlatformFile.h"
 
 namespace Trace
 {
