Removed date from copyright preamble.

#rnx

diff --git a/Engine/Source/Developer/TraceAnalysis/Private/Asio/Asio.cpp b/Engine/Source/Developer/TraceAnalysis/Private/Asio/Asio.cpp
--- a/Engine/Source/Developer/TraceAnalysis/Private/Asio/Asio.cpp
+++ b/Engine/Source/Developer/TraceAnalysis/Private/Asio/Asio.cpp
@@ -1,4 +1,4 @@
-// Copyright 1998-2019 Epic Games, Inc. All Rights Reserved.
+// Copyright Epic Games, Inc. All Rights Reserved.
 
 #include "Asio.h"
 
