Added TArrayView access support for array fields

Also fixed div-by-zero in Num() when a field does not exist

diff --git a/Engine/Source/Developer/TraceAnalysis/Private/Analysis/Engine.cpp b/Engine/Source/Developer/TraceAnalysis/Private/Analysis/Engine.cpp
--- a/Engine/Source/Developer/TraceAnalysis/Private/Analysis/Engine.cpp
+++ b/Engine/Source/Developer/TraceAnalysis/Private/Analysis/Engine.cpp
@@ -332,7 +332,7 @@ uint32 IAnalyzer::FArrayReader::Num() const
 	const auto* Inner = (const FAuxData*)this;
 	int32 SizeAndType = Inner->FieldSizeAndType;
 	SizeAndType = (SizeAndType < 0) ? -SizeAndType : SizeAndType;
-	return Inner->DataSize / SizeAndType;
+	return (SizeAndType == 0) ? SizeAndType : (Inner->DataSize / SizeAndType);
 }
 
 ////////////////////////////////////////////////////////////////////////////////
diff --git a/Engine/Source/Developer/TraceAnalysis/Public/Trace/Analyzer.h b/Engine/Source/Developer/TraceAnalysis/Public/Trace/Analyzer.h
--- a/Engine/Source/Developer/TraceAnalysis/Public/Trace/Analyzer.h
+++ b/Engine/Source/Developer/TraceAnalysis/Public/Trace/Analyzer.h
@@ -5,6 +5,7 @@
 #include "CoreTypes.h"
 #include "Logging/LogMacros.h"
 #include "Trace/Detail/Field.h"
+#include "Containers/ArrayView.h"
 
 #include <memory.h>
 
@@ -117,6 +118,9 @@ public:
 		/* Returns a element from the array an the given index or zero if Index
 		 * is out of bounds. */
 		ValueType operator [] (uint32 Index) const;
+
+		/** Get a pointer to the contiguous array data */
+		const ValueType* GetData() const;
 	};
 
 	struct TRACEANALYSIS_API FEventData
@@ -136,6 +140,12 @@ public:
 		 * @param FieldName The name of the event's field to get the value for. */
 		template <typename ValueType> const TArrayReader<ValueType>& GetArray(const ANSICHAR* FieldName) const;
 
+		/** Returns an array view for reading data from an array-type field. A valid
+		 * array view will always be returned even if no field matching the
+		 * given name was found.
+		 * @param FieldName The name of the event's field to get the value for. */
+		template <typename ValueType> TArrayView<const ValueType> GetArrayView(const ANSICHAR* FieldName) const;
+
 		/** Returns the event's attachment. Not that this will always return an
 		 * address but if the event has no attachment then reading from that
 		 * address if undefined. */
@@ -210,6 +220,7 @@ public:
 
 private:
 	template <typename ValueType> static ValueType CoerceValue(const void* Addr, int16 SizeAndType);
+	template <typename ValueType> static const ValueType* CoercePtr(const void* Addr, int16 SizeAndType);
 };
 
 ////////////////////////////////////////////////////////////////////////////////
@@ -234,6 +245,29 @@ ValueType IAnalyzer::CoerceValue(const void* Addr, int16 SizeAndType)
 	return ValueType(0);
 }
 
+////////////////////////////////////////////////////////////////////////////////
+template <typename ValueType>
+const ValueType* IAnalyzer::CoercePtr(const void* Addr, int16 SizeAndType)
+{
+	if (Addr == nullptr)
+	{
+		return nullptr;
+	}
+
+	switch (SizeAndType)
+	{
+	case -4:
+	case -8:
+	case  1:
+	case  2:
+	case  4:
+	case  8:
+		return reinterpret_cast<const ValueType*>(Addr);
+	}
+
+	return nullptr;
+}
+
 ////////////////////////////////////////////////////////////////////////////////
 template <typename ValueType>
 ValueType IAnalyzer::FEventData::GetValue(const ANSICHAR* FieldName) const
@@ -251,6 +285,14 @@ const IAnalyzer::TArrayReader<ValueType>& IAnalyzer::FEventData::GetArray(const
 	return *(TArrayReader<ValueType>*)(Base);
 }
 
+////////////////////////////////////////////////////////////////////////////////
+template <typename ValueType>
+TArrayView<const ValueType> IAnalyzer::FEventData::GetArrayView(const ANSICHAR* FieldName) const
+{
+	const TArrayReader<ValueType>& ArrayReader = GetArray<ValueType>(FieldName);
+	return TArrayView<const ValueType>(ArrayReader.GetData(), ArrayReader.Num());
+}
+
 ////////////////////////////////////////////////////////////////////////////////
 template <typename ValueType>
 ValueType IAnalyzer::TArrayReader<ValueType>::operator [] (uint32 Index) const
@@ -260,4 +302,13 @@ ValueType IAnalyzer::TArrayReader<ValueType>::operator [] (uint32 Index) const
 	return CoerceValue<ValueType>(Addr, ElementSizeAndType);
 }
 
+////////////////////////////////////////////////////////////////////////////////
+template <typename ValueType>
+const ValueType* IAnalyzer::TArrayReader<ValueType>::GetData() const
+{
+	int16 ElementSizeAndType;
+	const void* Addr = GetImpl(0, ElementSizeAndType);
+	return CoercePtr<ValueType>(Addr, ElementSizeAndType);
+}
+
 } // namespace Trace
