Added a QuickStoreHash() for string views.

#rnx

diff --git a/Engine/Source/Developer/TraceAnalysis/Private/Store/Utils.h b/Engine/Source/Developer/TraceAnalysis/Private/Store/Utils.h
--- a/Engine/Source/Developer/TraceAnalysis/Private/Store/Utils.h
+++ b/Engine/Source/Developer/TraceAnalysis/Private/Store/Utils.h
@@ -2,6 +2,8 @@
 
 #pragma once
 
+#include "Containers/StringView.h"
+
 namespace Trace
 {
 
@@ -17,5 +19,18 @@ inline uint32 QuickStoreHash(const T* String)
 	return Value;
 }
 
+////////////////////////////////////////////////////////////////////////////////
+template <typename CharType>
+inline uint32 QuickStoreHash(TStringViewImpl<CharType> View)
+{
+	const CharType* String = View.GetData();
+	uint32 Value = 5381;
+	for (int i = View.Len(); i > 0; --i, ++String)
+	{
+		Value = ((Value << 5) + Value) + *String;
+	}
+	return Value;
+}
+
 
 } // namespace Trace
