Temporarily clamp the number of async-IO threads while we're not thread safe.

#rnx

diff --git a/Engine/Source/Developer/TraceAnalysis/Private/Store/StoreService.cpp b/Engine/Source/Developer/TraceAnalysis/Private/Store/StoreService.cpp
--- a/Engine/Source/Developer/TraceAnalysis/Private/Store/StoreService.cpp
+++ b/Engine/Source/Developer/TraceAnalysis/Private/Store/StoreService.cpp
@@ -65,6 +65,9 @@ FStoreService* FStoreService::Create(const FDesc& InDesc)
 		Desc.ThreadCount = FPlatformMisc::NumberOfCoresIncludingHyperthreads();
 	}
 
+	// TODO: not thread safe yet
+	Desc.ThreadCount = 1;
+
 	FStoreService::FImpl* Impl = new FStoreService::FImpl(Desc);
 	if (Desc.RecorderPort >= 0)
 	{
