Specifying a file's attribute isn't necessary when opening an existing file.

#rnx

diff --git a/Engine/Source/Developer/TraceAnalysis/Private/Store/AsioFile.cpp b/Engine/Source/Developer/TraceAnalysis/Private/Store/AsioFile.cpp
--- a/Engine/Source/Developer/TraceAnalysis/Private/Store/AsioFile.cpp
+++ b/Engine/Source/Developer/TraceAnalysis/Private/Store/AsioFile.cpp
@@ -85,7 +85,7 @@ FAsioWriteable* FAsioFile::WriteFile(asio::io_context& IoContext, const TCHAR* P
 FAsioReadable* FAsioFile::ReadFile(asio::io_context& IoContext, const TCHAR* Path)
 {
 	HANDLE Handle = CreateFileW(Path, GENERIC_READ, FILE_SHARE_READ|FILE_SHARE_WRITE,
-		nullptr, OPEN_EXISTING, FILE_FLAG_OVERLAPPED|FILE_ATTRIBUTE_NORMAL, nullptr);
+		nullptr, OPEN_EXISTING, FILE_FLAG_OVERLAPPED, nullptr);
 	if (Handle == INVALID_HANDLE_VALUE)
 	{
 		return nullptr;
