RecorderPeer -> RecoderRelay

#rnx

diff --git a/Engine/Source/Developer/TraceAnalysis/Private/Store/AsioRecorder.cpp b/Engine/Source/Developer/TraceAnalysis/Private/Store/AsioRecorder.cpp
--- a/Engine/Source/Developer/TraceAnalysis/Private/Store/AsioRecorder.cpp
+++ b/Engine/Source/Developer/TraceAnalysis/Private/Store/AsioRecorder.cpp
@@ -12,12 +12,12 @@ namespace Trace
 {
 
 ////////////////////////////////////////////////////////////////////////////////
-class FAsioRecorderPeer
+class FAsioRecorderRelay
 	: public FAsioIoSink
 {
 public:
-						FAsioRecorderPeer(asio::ip::tcp::socket& Socket, FAsioWriteable* InOutput);
-	virtual				~FAsioRecorderPeer();
+						FAsioRecorderRelay(asio::ip::tcp::socket& Socket, FAsioWriteable* InOutput);
+	virtual				~FAsioRecorderRelay();
 	void				Close();
 
 private:
@@ -30,7 +30,7 @@ private:
 };
 
 ////////////////////////////////////////////////////////////////////////////////
-FAsioRecorderPeer::FAsioRecorderPeer(asio::ip::tcp::socket& Socket, FAsioWriteable* InOutput)
+FAsioRecorderRelay::FAsioRecorderRelay(asio::ip::tcp::socket& Socket, FAsioWriteable* InOutput)
 : Input(Socket)
 , Output(InOutput)
 {
@@ -38,20 +38,20 @@ FAsioRecorderPeer::FAsioRecorderPeer(asio::ip::tcp::socket& Socket, FAsioWriteab
 }
 
 ////////////////////////////////////////////////////////////////////////////////
-FAsioRecorderPeer::~FAsioRecorderPeer()
+FAsioRecorderRelay::~FAsioRecorderRelay()
 {
 	delete Output;
 }
 
 ////////////////////////////////////////////////////////////////////////////////
-void FAsioRecorderPeer::Close()
+void FAsioRecorderRelay::Close()
 {
 	Input.Close();
 	Output->Close();
 }
 
 ////////////////////////////////////////////////////////////////////////////////
-void FAsioRecorderPeer::OnIoComplete(uint32 Id, int32 Size)
+void FAsioRecorderRelay::OnIoComplete(uint32 Id, int32 Size)
 {
 	if (Size < 0)
 	{
@@ -86,10 +86,10 @@ FAsioRecorder::FAsioRecorder(asio::io_context& IoContext, FAsioStore& InStore)
 ////////////////////////////////////////////////////////////////////////////////
 FAsioRecorder::~FAsioRecorder()
 {
-	for (FAsioRecorderPeer* Peer : Peers)
+	for (FAsioRecorderRelay* Relay : Relays)
 	{
-		Peer->Close();
-		delete Peer;
+		Relay->Close();
+		delete Relay;
 	}
 }
 
@@ -102,8 +102,8 @@ bool FAsioRecorder::OnAccept(asio::ip::tcp::socket& Socket)
 		return true;
 	}
 
-	FAsioRecorderPeer* Peer = new FAsioRecorderPeer(Socket, Output);
-	Peers.Add(Peer);
+	FAsioRecorderRelay* Relay = new FAsioRecorderRelay(Socket, Output);
+	Relays.Add(Relay);
 	return true;
 }
 
diff --git a/Engine/Source/Developer/TraceAnalysis/Private/Store/AsioRecorder.h b/Engine/Source/Developer/TraceAnalysis/Private/Store/AsioRecorder.h
--- a/Engine/Source/Developer/TraceAnalysis/Private/Store/AsioRecorder.h
+++ b/Engine/Source/Developer/TraceAnalysis/Private/Store/AsioRecorder.h
@@ -14,7 +14,7 @@ namespace Trace
 {
 
 class FAsioStore;
-class FAsioRecorderPeer;
+class FAsioRecorderRelay;
 
 ////////////////////////////////////////////////////////////////////////////////
 class FAsioRecorder
@@ -28,7 +28,7 @@ public:
 private:
 	virtual bool				OnAccept(asio::ip::tcp::socket& Socket) override;
 	virtual void				OnTick() override;
-	TArray<FAsioRecorderPeer*>	Peers;
+	TArray<FAsioRecorderRelay*>	Relays;
 	FAsioStore&					Store;
 };
 
