Return trace ID when creating a new trace with the store.

#rnx

diff --git a/Engine/Source/Developer/TraceAnalysis/Private/Store/AsioRecorder.cpp b/Engine/Source/Developer/TraceAnalysis/Private/Store/AsioRecorder.cpp
--- a/Engine/Source/Developer/TraceAnalysis/Private/Store/AsioRecorder.cpp
+++ b/Engine/Source/Developer/TraceAnalysis/Private/Store/AsioRecorder.cpp
@@ -96,13 +96,13 @@ FAsioRecorder::~FAsioRecorder()
 ////////////////////////////////////////////////////////////////////////////////
 bool FAsioRecorder::OnAccept(asio::ip::tcp::socket& Socket)
 {
-	FAsioWriteable* Output = Store.CreateTrace();
-	if (Output == nullptr)
+	FAsioStore::FNewTrace Trace = Store.CreateTrace();
+	if (Trace.Writeable == nullptr)
 	{
 		return true;
 	}
 
-	FAsioRecorderRelay* Relay = new FAsioRecorderRelay(Socket, Output);
+	FAsioRecorderRelay* Relay = new FAsioRecorderRelay(Socket, Trace.Writeable);
 	Relays.Add(Relay);
 	return true;
 }
diff --git a/Engine/Source/Developer/TraceAnalysis/Private/Store/AsioStore.cpp b/Engine/Source/Developer/TraceAnalysis/Private/Store/AsioStore.cpp
--- a/Engine/Source/Developer/TraceAnalysis/Private/Store/AsioStore.cpp
+++ b/Engine/Source/Developer/TraceAnalysis/Private/Store/AsioStore.cpp
@@ -18,33 +18,22 @@ FAsioStore::FAsioStore(asio::io_context& IoContext, const TCHAR* InStoreDir)
 , StoreDir(InStoreDir)
 {
 	IPlatformFile& FileSystem = IPlatformFile::GetPlatformPhysical();
-	FileSystem.IterateDirectory(InStoreDir, [this] (const TCHAR* Name, bool IsDirectory)
+	FileSystem.IterateDirectory(InStoreDir, [this] (const TCHAR* Path, bool IsDirectory)
 	{
 #if 0
 		if (IsDirectory)
 		{
-			int32 Id = FCString::Atoi(Name);
+			int32 Id = FCString::Atoi(Path);
 			LastTraceId = (Id < LastTraceId) ? Id : LastTraceId;
 		}
 #else
-		const TCHAR* Dot = FCString::Strrchr(Name, '.');
+		const TCHAR* Dot = FCString::Strrchr(Path, '.');
 		if (Dot == nullptr || FCString::Strcmp(Dot, TEXT(".utrace")))
 		{
 			return true;
 		}
 
-		for (const TCHAR* c = Dot; c > Name; --c)
-		{
-			if (c[-1] == '\\' || c[-1] == '/')
-			{
-				Name = c;
-				break;
-			}
-		}
-
-		FTrace* Trace = new FTrace();
-		Trace->Name.AppendChars(Name, int32(Dot - Name));
-		Traces.Add(Trace);
+		AddTrace(Path);
 #endif // 0
 
 		return true;
@@ -92,7 +81,34 @@ FAsioStore::FTrace* FAsioStore::GetTrace(uint32 Id)
 }
 
 ////////////////////////////////////////////////////////////////////////////////
-FAsioWriteable* FAsioStore::CreateTrace()
+FAsioStore::FTrace* FAsioStore::AddTrace(const TCHAR* Path)
+{
+	const TCHAR* Dot = FCString::Strrchr(Path, '.');
+	if (Dot == nullptr)
+	{
+		return nullptr;
+	}
+
+	FStringView Name;
+	for (const TCHAR* c = Dot; c > Path; --c)
+	{
+		if (c[-1] == '\\' || c[-1] == '/')
+		{
+			Name = FStringView(c, int32(Dot - c));
+			break;
+		}
+	}
+
+	FTrace* Trace = new FTrace();
+	Trace->Name = Name;
+	Trace->Id = QuickStoreHash(Name);
+
+	Traces.Add(Trace);
+	return Trace;
+}
+
+////////////////////////////////////////////////////////////////////////////////
+FAsioStore::FNewTrace FAsioStore::CreateTrace()
 {
 	IPlatformFile& PlatformFile = IPlatformFile::GetPlatformPhysical();
 
@@ -115,7 +131,7 @@ FAsioWriteable* FAsioStore::CreateTrace()
 
 	if (!bOk)
 	{
-		return nullptr;
+		return {};
 	}
 
 	TracePath += TEXT("/data.utrace");
@@ -131,7 +147,20 @@ FAsioWriteable* FAsioStore::CreateTrace()
 	}
 #endif // 0
 
-	return FAsioFile::WriteFile(GetIoContext(), *TracePath);
+	FAsioWriteable* File = FAsioFile::WriteFile(GetIoContext(), *TracePath);
+	if (File == nullptr)
+	{
+		return {};
+	}
+
+	FTrace* Trace = AddTrace(*TracePath);
+	if (Trace == nullptr)
+	{
+		delete File;
+		return {};
+	}
+
+	return { Trace->GetId(), File };
 }
 
 ////////////////////////////////////////////////////////////////////////////////
diff --git a/Engine/Source/Developer/TraceAnalysis/Private/Store/AsioStore.h b/Engine/Source/Developer/TraceAnalysis/Private/Store/AsioStore.h
--- a/Engine/Source/Developer/TraceAnalysis/Private/Store/AsioStore.h
+++ b/Engine/Source/Developer/TraceAnalysis/Private/Store/AsioStore.h
@@ -39,6 +39,13 @@ public:
 	private:
 		friend			FAsioStore;
 		FString			Name;
+		uint32			Id;
+	};
+
+	struct FNewTrace
+	{
+		uint32			Id;
+		FAsioWriteable* Writeable;
 	};
 
 						FAsioStore(asio::io_context& IoContext, const TCHAR* InStoreDir);
@@ -46,11 +53,12 @@ public:
 	uint32				GetTraceCount() const;
 	const FTrace*		GetTraceInfo(uint32 Index) const;
 	int32				ReadTrace(uint32 Id);
-	FAsioWriteable*		CreateTrace();
+	FNewTrace			CreateTrace();
 	FAsioReadable*		OpenTrace(uint32 Id);
 
 private:
 	FTrace*				GetTrace(uint32 Id);
+	FTrace*				AddTrace(const TCHAR* Path);
 	FString				StoreDir;
 	TArray<FTrace*>		Traces;
 #if 0
