Moved some methods from header to source.

#rnx

diff --git a/Engine/Source/Developer/TraceAnalysis/Private/Store/AsioStore.cpp b/Engine/Source/Developer/TraceAnalysis/Private/Store/AsioStore.cpp
--- a/Engine/Source/Developer/TraceAnalysis/Private/Store/AsioStore.cpp
+++ b/Engine/Source/Developer/TraceAnalysis/Private/Store/AsioStore.cpp
@@ -12,6 +12,26 @@
 namespace Trace
 {
 
+////////////////////////////////////////////////////////////////////////////////
+const TCHAR* FAsioStore::FTrace::GetName() const
+{
+	return *Name;
+}
+
+////////////////////////////////////////////////////////////////////////////////
+uint32 FAsioStore::FTrace::GetId() const
+{
+	return Id;
+}
+
+////////////////////////////////////////////////////////////////////////////////
+uint64 FAsioStore::FTrace::GetSize() const
+{
+	return 0;
+}
+
+
+
 ////////////////////////////////////////////////////////////////////////////////
 FAsioStore::FAsioStore(asio::io_context& IoContext, const TCHAR* InStoreDir)
 : FAsioObject(IoContext)
diff --git a/Engine/Source/Developer/TraceAnalysis/Private/Store/AsioStore.h b/Engine/Source/Developer/TraceAnalysis/Private/Store/AsioStore.h
--- a/Engine/Source/Developer/TraceAnalysis/Private/Store/AsioStore.h
+++ b/Engine/Source/Developer/TraceAnalysis/Private/Store/AsioStore.h
@@ -32,9 +32,9 @@ public:
 	class FTrace
 	{
 	public:
-		const TCHAR*	GetName() const		{ return *Name; }
-		uint32			GetId() const		{ return QuickStoreHash(GetName()); }
-		uint64			GetSize() const		{ return 3; }
+		const TCHAR*	GetName() const;
+		uint32			GetId() const;
+		uint64			GetSize() const;
 
 	private:
 		friend			FAsioStore;
