Use a static Connect() method instead of directly constructing a store client

#rnx

diff --git a/Engine/Source/Developer/TraceAnalysis/Private/Store/StoreClient.cpp b/Engine/Source/Developer/TraceAnalysis/Private/Store/StoreClient.cpp
--- a/Engine/Source/Developer/TraceAnalysis/Private/Store/StoreClient.cpp
+++ b/Engine/Source/Developer/TraceAnalysis/Private/Store/StoreClient.cpp
@@ -271,17 +271,20 @@ struct FStoreClient::FImpl
 };
 
 ////////////////////////////////////////////////////////////////////////////////
-FStoreClient::FStoreClient(const TCHAR* Host, uint16 Port)
-: Impl(nullptr)
+FStoreClient* FStoreClient::Connect(const TCHAR* Host, uint32 Port)
 {
-	Impl = new FStoreClient::FImpl();
+	FImpl* Impl = new FStoreClient::FImpl();
 #if TRACE_WITH_ASIO
 	if (!Impl->Connect(Host, Port))
 #endif
 	{
 		delete Impl;
-		Impl = nullptr;
+		return nullptr;
 	}
+
+	FStoreClient* Client = new FStoreClient();
+	Client->Impl = Impl;
+	return Client;
 }
 
 ////////////////////////////////////////////////////////////////////////////////
@@ -300,7 +303,7 @@ bool FStoreClient::IsValid() const
 const FStoreClient::FStatus* FStoreClient::GetStatus()
 {
 #if TRACE_WITH_ASIO
-	if (Impl == nullptr || !Impl->GetStatus())
+	if (!Impl->GetStatus())
 	{
 		return nullptr;
 	}
@@ -316,7 +319,7 @@ const FStoreClient::FStatus* FStoreClient::GetStatus()
 uint32 FStoreClient::GetTraceCount()
 {
 #if TRACE_WITH_ASIO
-	if (Impl == nullptr || !Impl->GetTraceCount())
+	if (!Impl->GetTraceCount())
 	{
 		return 0;
 	}
@@ -331,7 +334,7 @@ uint32 FStoreClient::GetTraceCount()
 const FStoreClient::FTraceInfo* FStoreClient::GetTraceInfo(uint32 Index)
 {
 #if TRACE_WITH_ASIO
-	if (Impl == nullptr || !Impl->GetTraceInfo(Index))
+	if (!Impl->GetTraceInfo(Index))
 	{
 		return nullptr;
 	}
@@ -347,11 +350,6 @@ const FStoreClient::FTraceInfo* FStoreClient::GetTraceInfo(uint32 Index)
 bool FStoreClient::ReadTrace(uint32 Id, FTraceData& Out)
 {
 #if TRACE_WITH_ASIO
-	if (Impl == nullptr)
-	{
-		return false;
-	}
-
 	Out.Handle = UPTRINT(Impl->ReadTrace(Id));
 	return (Out.Handle != 0);
 #else
diff --git a/Engine/Source/Developer/TraceAnalysis/Public/Trace/StoreClient.h b/Engine/Source/Developer/TraceAnalysis/Public/Trace/StoreClient.h
--- a/Engine/Source/Developer/TraceAnalysis/Public/Trace/StoreClient.h
+++ b/Engine/Source/Developer/TraceAnalysis/Public/Trace/StoreClient.h
@@ -34,8 +34,9 @@ public:
 		friend			FStoreClient;
 		UPTRINT			Handle = 0;
 	};
-						FStoreClient(const TCHAR* Host, uint16 Port);
+	
 						~FStoreClient();
+	static FStoreClient*Connect(const TCHAR* Host, uint32 Port);
 	bool				IsValid() const;
 	const FStatus*		GetStatus();
 	uint32				GetTraceCount();
@@ -71,8 +72,9 @@ public:
 #endif // 0
 
 private:
+						FStoreClient() = default;
 	struct				FImpl;
-	FImpl*				Impl;
+	FImpl*				Impl = nullptr;
 
 private:
 						FStoreClient(const FStoreClient&) = delete;
