StoreClient::ReadTrace() returns an IInDataStream directly instead of some convoluted abstraction

#rnx

diff --git a/Engine/Source/Developer/TraceAnalysis/Private/Store/StoreClient.cpp b/Engine/Source/Developer/TraceAnalysis/Private/Store/StoreClient.cpp
--- a/Engine/Source/Developer/TraceAnalysis/Private/Store/StoreClient.cpp
+++ b/Engine/Source/Developer/TraceAnalysis/Private/Store/StoreClient.cpp
@@ -7,10 +7,61 @@
 
 #include "AsioStore.h"
 #include "CborPayload.h"
+#include "Templates/UnrealTemplate.h"
+#include "Trace/DataStream.h"
 
 namespace Trace
 {
 
+////////////////////////////////////////////////////////////////////////////////
+class FTraceDataStream
+	: public IInDataStream
+{
+public:
+							FTraceDataStream(asio::ip::tcp::socket& InSocket);
+	virtual					~FTraceDataStream();
+	virtual int32			Read(void* Dest, uint32 DestSize) override;
+
+private:
+#if TRACE_WITH_ASIO
+	asio::ip::tcp::socket	Socket;
+#endif // TRACE_WITH_ASIO
+};
+
+////////////////////////////////////////////////////////////////////////////////
+FTraceDataStream::FTraceDataStream(asio::ip::tcp::socket& InSocket)
+: Socket(MoveTemp(InSocket))
+{
+};
+
+////////////////////////////////////////////////////////////////////////////////
+FTraceDataStream::~FTraceDataStream()
+{
+#if TRACE_WITH_ASIO
+	Socket.close();
+#endif // TRACE_WITH_ASIO
+}
+
+////////////////////////////////////////////////////////////////////////////////
+int32 FTraceDataStream::Read(void* Dest, uint32 DestSize)
+{
+#if TRACE_WITH_ASIO
+	asio::error_code ErrorCode;
+	size_t BytesRead = Socket.read_some(asio::buffer(Dest, DestSize), ErrorCode);
+	if (ErrorCode)
+	{
+		Socket.close();
+		return -1;
+	}
+
+	return int32(BytesRead);
+#else
+	return -1;
+#endif // TRACE_WITH_ASIO
+}
+
+
+
 ////////////////////////////////////////////////////////////////////////////////
 class FStoreCborClient
 {
@@ -22,7 +73,7 @@ public:
 	bool					GetStatus();
 	bool					GetTraceCount();
 	bool					GetTraceInfo(uint32 Index);
-	asio::ip::tcp::socket*	ReadTrace(uint32 Id);
+	FTraceDataStream*		ReadTrace(uint32 Id);
 
 private:
 	bool					Communicate(const FPayload& Payload);
@@ -144,7 +195,7 @@ bool FStoreCborClient::GetTraceInfo(uint32 Index)
 }
 
 ////////////////////////////////////////////////////////////////////////////////
-asio::ip::tcp::socket* FStoreCborClient::ReadTrace(uint32 Id)
+FTraceDataStream* FStoreCborClient::ReadTrace(uint32 Id)
 {
 	TPayloadBuilder<> Builder("trace_read");
 	Builder.AddParam("id", Id);
@@ -171,7 +222,7 @@ asio::ip::tcp::socket* FStoreCborClient::ReadTrace(uint32 Id)
 		return nullptr;
 	}
 
-	return new asio::ip::tcp::socket(MoveTemp(SenderSocket));
+	return new FTraceDataStream(SenderSocket);
 }
 
 } // namespace Trace
@@ -218,50 +269,6 @@ uint64 FStoreClient::FTraceInfo::GetSize() const
 
 
 
-////////////////////////////////////////////////////////////////////////////////
-FStoreClient::FTraceData::~FTraceData()
-{
-#if TRACE_WITH_ASIO
-	if (auto* Inner = (asio::ip::tcp::socket*)Handle)
-	{
-		Inner->close();
-		delete Inner;
-	}
-#endif // TRACE_WITH_ASIO
-}
-
-////////////////////////////////////////////////////////////////////////////////
-int32 FStoreClient::FTraceData::Read(void* Dest, uint32 DestSize)
-{
-#if TRACE_WITH_ASIO
-	auto* Inner = (asio::ip::tcp::socket*)Handle;
-	if (Inner == nullptr)
-	{
-		return -1;
-	}
-
-	asio::error_code ErrorCode;
-	size_t BytesRead = Inner->read_some(asio::buffer(Dest, DestSize), ErrorCode);
-	if (ErrorCode)
-	{
-		Inner->close();
-		return -1;
-	}
-
-	return int32(BytesRead);
-#else
-	return -1;
-#endif // TRACE_WITH_ASIO
-}
-
-////////////////////////////////////////////////////////////////////////////////
-bool FStoreClient::FTraceData::IsValid() const
-{
-	return (Handle != 0);
-}
-
-
-
 ////////////////////////////////////////////////////////////////////////////////
 struct FStoreClient::FImpl
 #if TRACE_WITH_ASIO
@@ -347,14 +354,9 @@ const FStoreClient::FTraceInfo* FStoreClient::GetTraceInfo(uint32 Index)
 }
 
 ////////////////////////////////////////////////////////////////////////////////
-bool FStoreClient::ReadTrace(uint32 Id, FTraceData& Out)
+FStoreClient::FTraceData FStoreClient::ReadTrace(uint32 Id)
 {
-#if TRACE_WITH_ASIO
-	Out.Handle = UPTRINT(Impl->ReadTrace(Id));
-	return (Out.Handle != 0);
-#else
-	return false;
-#endif // TRACE_WITH_ASIO
+	return FTraceData(Impl->ReadTrace(Id));
 }
 
 } // namespace Trace
diff --git a/Engine/Source/Developer/TraceAnalysis/Public/Trace/StoreClient.h b/Engine/Source/Developer/TraceAnalysis/Public/Trace/StoreClient.h
--- a/Engine/Source/Developer/TraceAnalysis/Public/Trace/StoreClient.h
+++ b/Engine/Source/Developer/TraceAnalysis/Public/Trace/StoreClient.h
@@ -7,6 +7,8 @@
 namespace Trace
 {
 
+class IInDataStream;
+
 ////////////////////////////////////////////////////////////////////////////////
 class TRACEANALYSIS_API FStoreClient
 {
@@ -24,15 +26,10 @@ public:
 		//template <typename Lambda> uint32 ReadMetadata(Lambda&& Callback) const;
 	};
 
-	struct TRACEANALYSIS_API FTraceData
+	struct FTraceData
+		: public TUniquePtr<IInDataStream>
 	{
-						~FTraceData();
-		int32			Read(void* Dest, uint32 DestSize);
-		bool			IsValid() const;
-
-	private:
-		friend			FStoreClient;
-		UPTRINT			Handle = 0;
+		using TUniquePtr<IInDataStream>::TUniquePtr;
 	};
 	
 						~FStoreClient();
@@ -41,7 +38,7 @@ public:
 	const FStatus*		GetStatus();
 	uint32				GetTraceCount();
 	const FTraceInfo*	GetTraceInfo(uint32 Index);
-	bool				ReadTrace(uint32 Id, FTraceData& Out);
+	FTraceData			ReadTrace(uint32 Id);
 
 #if 0
 	template <typename Lambda> uint32 GetTraceInfos(uint32 StartIndex, uint32 Count, Lambda&& Callback) const;
