Keep track of which trace a recorder session is recording to

#rnx

diff --git a/Engine/Source/Developer/TraceAnalysis/Private/Store/AsioRecorder.cpp b/Engine/Source/Developer/TraceAnalysis/Private/Store/AsioRecorder.cpp
--- a/Engine/Source/Developer/TraceAnalysis/Private/Store/AsioRecorder.cpp
+++ b/Engine/Source/Developer/TraceAnalysis/Private/Store/AsioRecorder.cpp
@@ -86,10 +86,10 @@ FAsioRecorder::FAsioRecorder(asio::io_context& IoContext, FAsioStore& InStore)
 ////////////////////////////////////////////////////////////////////////////////
 FAsioRecorder::~FAsioRecorder()
 {
-	for (FAsioRecorderRelay* Relay : Relays)
+	for (FSession& Session : Sessions)
 	{
-		Relay->Close();
-		delete Relay;
+		Session.Relay->Close();
+		delete Session.Relay;
 	}
 }
 
@@ -103,7 +103,7 @@ bool FAsioRecorder::OnAccept(asio::ip::tcp::socket& Socket)
 	}
 
 	FAsioRecorderRelay* Relay = new FAsioRecorderRelay(Socket, Trace.Writeable);
-	Relays.Add(Relay);
+	Sessions.Add({Relay, Trace.Id});
 	return true;
 }
 
diff --git a/Engine/Source/Developer/TraceAnalysis/Private/Store/AsioRecorder.h b/Engine/Source/Developer/TraceAnalysis/Private/Store/AsioRecorder.h
--- a/Engine/Source/Developer/TraceAnalysis/Private/Store/AsioRecorder.h
+++ b/Engine/Source/Developer/TraceAnalysis/Private/Store/AsioRecorder.h
@@ -26,9 +26,15 @@ public:
 								~FAsioRecorder();
 
 private:
+	struct FSession
+	{
+		FAsioRecorderRelay*		Relay;
+		uint32					TraceId;
+	};
+
 	virtual bool				OnAccept(asio::ip::tcp::socket& Socket) override;
 	virtual void				OnTick() override;
-	TArray<FAsioRecorderRelay*>	Relays;
+	TArray<FSession>			Sessions;
 	FAsioStore&					Store;
 };
 
