Use explicitly named methods instead of overloading for the API to build CBOR-based payloads

#rnx

diff --git a/Engine/Source/Developer/TraceAnalysis/Private/Store/AsioStoreCborServer.cpp b/Engine/Source/Developer/TraceAnalysis/Private/Store/AsioStoreCborServer.cpp
--- a/Engine/Source/Developer/TraceAnalysis/Private/Store/AsioStoreCborServer.cpp
+++ b/Engine/Source/Developer/TraceAnalysis/Private/Store/AsioStoreCborServer.cpp
@@ -90,7 +90,7 @@ void FAsioStoreCborPeer::Close()
 void FAsioStoreCborPeer::OnConnect()
 {
 	TPayloadBuilder<> Builder((int32)EStatusCode::Success);
-	Builder.AddParam("version", int32(EStoreVersion::Value));
+	Builder.AddInteger("version", int32(EStoreVersion::Value));
 	SendResponse(Builder.Done());
 }
 
@@ -98,7 +98,7 @@ void FAsioStoreCborPeer::OnConnect()
 void FAsioStoreCborPeer::OnStatus()
 {
 	TPayloadBuilder<> Builder((int32)EStatusCode::Success);
-	Builder.AddParam("recorder_port", Recorder.GetPort());
+	Builder.AddInteger("recorder_port", Recorder.GetPort());
 	SendResponse(Builder.Done());
 }
 
@@ -106,7 +106,7 @@ void FAsioStoreCborPeer::OnStatus()
 void FAsioStoreCborPeer::OnTraceCount()
 {
 	TPayloadBuilder<> Builder((int32)EStatusCode::Success);
-	Builder.AddParam("count", Store.GetTraceCount());
+	Builder.AddInteger("count", Store.GetTraceCount());
 	SendResponse(Builder.Done());
 }
 
@@ -126,8 +126,8 @@ void FAsioStoreCborPeer::OnTraceInfo()
 	}
 
 	TPayloadBuilder<> Builder((int32)EStatusCode::Success);
-	Builder.AddParam("id", Trace->GetId());
-	Builder.AddParam("size", Trace->GetSize());
+	Builder.AddInteger("id", Trace->GetId());
+	Builder.AddInteger("size", Trace->GetSize());
 	SendResponse(Builder.Done());
 }
 
@@ -155,7 +155,7 @@ void FAsioStoreCborPeer::OnTraceRead()
 	}
 
 	TPayloadBuilder<> Builder((int32)EStatusCode::Success);
-	Builder.AddParam("port", Port);
+	Builder.AddInteger("port", Port);
 	SendResponse(Builder.Done());
 }
 
diff --git a/Engine/Source/Developer/TraceAnalysis/Private/Store/CborPayload.h b/Engine/Source/Developer/TraceAnalysis/Private/Store/CborPayload.h
--- a/Engine/Source/Developer/TraceAnalysis/Private/Store/CborPayload.h
+++ b/Engine/Source/Developer/TraceAnalysis/Private/Store/CborPayload.h
@@ -60,15 +60,15 @@ template <int Size=128>
 class TPayloadBuilder
 {
 public:
-									TPayloadBuilder(int32 StatusCode);
-	template <int N>				TPayloadBuilder(const char (&MethodName)[N]);
-	template <int N> void			AddParam(const char (&Name)[N], int64 Value);
-	template <int N, int M> void	AddParam(const char (&Name)[N], const char (&Value)[M]);
-	FPayload						Done();
+								TPayloadBuilder(int32 StatusCode);
+	template <int N>			TPayloadBuilder(const char (&MethodName)[N]);
+	template <int N> void		AddInteger(const char (&Name)[N], int64 Value);
+	template <int N> void		AddString(const char (&Name)[N], const char* Value, int32 Length=-1);
+	FPayload					Done();
 
 private:
-	TInlineMemoryWriter<Size>		MemoryWriter;
-	FCborWriter						CborWriter = &MemoryWriter;
+	TInlineMemoryWriter<Size>	MemoryWriter;
+	FCborWriter					CborWriter = &MemoryWriter;
 };
 
 ////////////////////////////////////////////////////////////////////////////////
@@ -76,7 +76,7 @@ template <int Size>
 inline TPayloadBuilder<Size>::TPayloadBuilder(int32 StatusCode)
 {
 	CborWriter.WriteContainerStart(ECborCode::Map, -1);
-	AddParam("$status", StatusCode);
+	AddInteger("$status", StatusCode);
 }
 
 ////////////////////////////////////////////////////////////////////////////////
@@ -85,13 +85,13 @@ template <int N>
 inline TPayloadBuilder<Size>::TPayloadBuilder(const char (&MethodName)[N])
 {
 	CborWriter.WriteContainerStart(ECborCode::Map, -1);
-	AddParam("$method", MethodName);
+	AddString("$method", MethodName, N);
 }
 
 ////////////////////////////////////////////////////////////////////////////////
 template <int Size>
 template <int N>
-inline void TPayloadBuilder<Size>::AddParam(const char (&Name)[N], int64 Value)
+inline void TPayloadBuilder<Size>::AddInteger(const char (&Name)[N], int64 Value)
 {
 	CborWriter.WriteValue(Name, N);
 	CborWriter.WriteValue(Value);
@@ -99,11 +99,15 @@ inline void TPayloadBuilder<Size>::AddParam(const char (&Name)[N], int64 Value)
 
 ////////////////////////////////////////////////////////////////////////////////
 template <int Size>
-template <int N, int M>
-inline void TPayloadBuilder<Size>::AddParam(const char (&Name)[N], const char (&Value)[M])
+template <int N>
+inline void TPayloadBuilder<Size>::AddString(
+	const char (&Name)[N],
+	const char* Value,
+	int Length)
 {
+	Length = (Length < 0) ? int32(strlen(Value)) : Length;
 	CborWriter.WriteValue(Name, N);
-	CborWriter.WriteValue(Value, M);
+	CborWriter.WriteValue(Value, Length);
 }
 
 ////////////////////////////////////////////////////////////////////////////////
diff --git a/Engine/Source/Developer/TraceAnalysis/Private/Store/StoreClient.cpp b/Engine/Source/Developer/TraceAnalysis/Private/Store/StoreClient.cpp
--- a/Engine/Source/Developer/TraceAnalysis/Private/Store/StoreClient.cpp
+++ b/Engine/Source/Developer/TraceAnalysis/Private/Store/StoreClient.cpp
@@ -118,7 +118,7 @@ bool FStoreCborClient::Connect(const TCHAR* Host, uint16 Port)
 	}
 
 	TPayloadBuilder<> Builder("connect");
-	Builder.AddParam("version", int32(EStoreVersion::Value));
+	Builder.AddInteger("version", int32(EStoreVersion::Value));
 	FPayload Payload = Builder.Done();
 	return Communicate(Payload);
 }
@@ -189,7 +189,7 @@ bool FStoreCborClient::GetTraceCount()
 bool FStoreCborClient::GetTraceInfo(uint32 Index)
 {
 	TPayloadBuilder<> Builder("trace_info");
-	Builder.AddParam("index", Index);
+	Builder.AddInteger("index", Index);
 	FPayload Payload = Builder.Done();
 	return Communicate(Payload);
 }
@@ -198,7 +198,7 @@ bool FStoreCborClient::GetTraceInfo(uint32 Index)
 FTraceDataStream* FStoreCborClient::ReadTrace(uint32 Id)
 {
 	TPayloadBuilder<> Builder("trace_read");
-	Builder.AddParam("id", Id);
+	Builder.AddInteger("id", Id);
 	FPayload Payload = Builder.Done();
 	if (!Communicate(Payload))
 	{
