Return strings as views as CBOR strings might not be null-terminated.

#rnx

diff --git a/Engine/Source/Developer/TraceAnalysis/Private/Store/AsioStoreCborServer.cpp b/Engine/Source/Developer/TraceAnalysis/Private/Store/AsioStoreCborServer.cpp
--- a/Engine/Source/Developer/TraceAnalysis/Private/Store/AsioStoreCborServer.cpp
+++ b/Engine/Source/Developer/TraceAnalysis/Private/Store/AsioStoreCborServer.cpp
@@ -162,8 +162,8 @@ void FAsioStoreCborPeer::OnTraceRead()
 ////////////////////////////////////////////////////////////////////////////////
 void FAsioStoreCborPeer::OnPayload()
 {
-	const char* Method = Response.GetString("$method");
-	if (Method[0] == '\0')
+	FAnsiStringView Method = Response.GetString("$method", "");
+	if (!Method.Len())
 	{
 		SendError(EStatusCode::BadRequest);
 		return;
diff --git a/Engine/Source/Developer/TraceAnalysis/Private/Store/CborPayload.h b/Engine/Source/Developer/TraceAnalysis/Private/Store/CborPayload.h
--- a/Engine/Source/Developer/TraceAnalysis/Private/Store/CborPayload.h
+++ b/Engine/Source/Developer/TraceAnalysis/Private/Store/CborPayload.h
@@ -8,6 +8,7 @@
 
 #include "CborReader.h"
 #include "CborWriter.h"
+#include "Containers/StringView.h"
 #include "Serialization/MemoryArchive.h"
 #include "Serialization/MemoryReader.h"
 
@@ -124,8 +125,9 @@ inline FPayload TPayloadBuilder<Size>::Done()
 class FResponse
 {
 public:
-	int64			GetInteger(const char* Key, int64 Default=0) const;
-	const char*		GetString(const char* Key, const char* Default="") const;
+	int64			GetInteger(const char* Key, int64 Default) const;
+	template <int N>
+	FAnsiStringView	GetString(const char* Key, const char (&Default)[N]) const;
 	const uint8*	GetData() const;
 	uint32			GetSize() const;
 	uint8*			Reserve(uint32 Size);
@@ -206,19 +208,23 @@ inline int64 FResponse::GetInteger(const char* Key, int64 Default) const
 }
 
 ////////////////////////////////////////////////////////////////////////////////
-inline const char* FResponse::GetString(const char* Key, const char* Default) const
+template <int N>
+inline FAnsiStringView FResponse::GetString(const char* Key, const char (&Default)[N]) const
 {
+	FAnsiStringView DefaultView(Default, N);
 	return GetValue(
 		Key,
-		Default,
-		[this, Default] (const FCborContext& Context, uint32 Offset)
+		DefaultView,
+		[this, DefaultView] (const FCborContext& Context, uint32 Offset)
 		{
 			if (Context.IsString())
 			{
-				return (const char*)(Buffer.GetData() + Offset - Context.AsLength());
+				int32 Length = Context.AsLength();
+				const char* Data = (const char*)(Buffer.GetData() + Offset - Length);
+				return FAnsiStringView(Data, Length - 1);
 			}
 
-			return Default;
+			return DefaultView;
 		}
 	);
 }
