Renamed method names so that they are a little REST-like

#rnx

diff --git a/Engine/Source/Developer/TraceAnalysis/Private/Store/AsioStoreCborServer.cpp b/Engine/Source/Developer/TraceAnalysis/Private/Store/AsioStoreCborServer.cpp
--- a/Engine/Source/Developer/TraceAnalysis/Private/Store/AsioStoreCborServer.cpp
+++ b/Engine/Source/Developer/TraceAnalysis/Private/Store/AsioStoreCborServer.cpp
@@ -169,15 +169,15 @@ void FAsioStoreCborPeer::OnPayload()
 		return;
 	}
 
-	struct {
+	static struct {
 		uint32	Hash;
 		void	(FAsioStoreCborPeer::*Func)();
 	} const DispatchTable[] = {
 		{ QuickStoreHash("connect"),	&FAsioStoreCborPeer::OnConnect },
 		{ QuickStoreHash("status"),		&FAsioStoreCborPeer::OnStatus },
-		{ QuickStoreHash("trace_count"),&FAsioStoreCborPeer::OnTraceCount },
-		{ QuickStoreHash("trace_info"),	&FAsioStoreCborPeer::OnTraceInfo },
-		{ QuickStoreHash("trace_read"),	&FAsioStoreCborPeer::OnTraceRead },
+		{ QuickStoreHash("trace/count"),&FAsioStoreCborPeer::OnTraceCount },
+		{ QuickStoreHash("trace/info"),	&FAsioStoreCborPeer::OnTraceInfo },
+		{ QuickStoreHash("trace/read"),	&FAsioStoreCborPeer::OnTraceRead },
 	};
 
 	uint32 MethodHash = QuickStoreHash(Method);
diff --git a/Engine/Source/Developer/TraceAnalysis/Private/Store/StoreClient.cpp b/Engine/Source/Developer/TraceAnalysis/Private/Store/StoreClient.cpp
--- a/Engine/Source/Developer/TraceAnalysis/Private/Store/StoreClient.cpp
+++ b/Engine/Source/Developer/TraceAnalysis/Private/Store/StoreClient.cpp
@@ -180,7 +180,7 @@ bool FStoreCborClient::GetStatus()
 ////////////////////////////////////////////////////////////////////////////////
 bool FStoreCborClient::GetTraceCount()
 {
-	TPayloadBuilder<32> Builder("trace_count");
+	TPayloadBuilder<32> Builder("trace/count");
 	FPayload Payload = Builder.Done();
 	return Communicate(Payload);
 }
@@ -188,7 +188,7 @@ bool FStoreCborClient::GetTraceCount()
 ////////////////////////////////////////////////////////////////////////////////
 bool FStoreCborClient::GetTraceInfo(uint32 Index)
 {
-	TPayloadBuilder<> Builder("trace_info");
+	TPayloadBuilder<> Builder("trace/info");
 	Builder.AddInteger("index", Index);
 	FPayload Payload = Builder.Done();
 	return Communicate(Payload);
@@ -197,7 +197,7 @@ bool FStoreCborClient::GetTraceInfo(uint32 Index)
 ////////////////////////////////////////////////////////////////////////////////
 FTraceDataStream* FStoreCborClient::ReadTrace(uint32 Id)
 {
-	TPayloadBuilder<> Builder("trace_read");
+	TPayloadBuilder<> Builder("trace/read");
 	Builder.AddInteger("id", Id);
 	FPayload Payload = Builder.Done();
 	if (!Communicate(Payload))
