Implemented FTraceInfo::GetSize()

#rnx

diff --git a/Engine/Source/Developer/TraceAnalysis/Private/Store/AsioStore.cpp b/Engine/Source/Developer/TraceAnalysis/Private/Store/AsioStore.cpp
--- a/Engine/Source/Developer/TraceAnalysis/Private/Store/AsioStore.cpp
+++ b/Engine/Source/Developer/TraceAnalysis/Private/Store/AsioStore.cpp
@@ -27,7 +27,10 @@ uint32 FAsioStore::FTrace::GetId() const
 ////////////////////////////////////////////////////////////////////////////////
 uint64 FAsioStore::FTrace::GetSize() const
 {
-	return 0;
+	HANDLE Inner = HANDLE(Handle);
+	LARGE_INTEGER FileSize;
+	GetFileSizeEx(Inner, &FileSize);
+	return FileSize.QuadPart;
 }
 
 
@@ -103,6 +106,13 @@ FAsioStore::FTrace* FAsioStore::GetTrace(uint32 Id)
 ////////////////////////////////////////////////////////////////////////////////
 FAsioStore::FTrace* FAsioStore::AddTrace(const TCHAR* Path)
 {
+	HANDLE Handle = CreateFileW(Path, 0, FILE_SHARE_READ|FILE_SHARE_WRITE|FILE_SHARE_DELETE,
+		nullptr, OPEN_EXISTING, 0, nullptr);
+	if (Handle == INVALID_HANDLE_VALUE)
+	{
+		return nullptr;
+	}
+
 	const TCHAR* Dot = FCString::Strrchr(Path, '.');
 	if (Dot == nullptr)
 	{
@@ -122,6 +132,7 @@ FAsioStore::FTrace* FAsioStore::AddTrace(const TCHAR* Path)
 	FTrace* Trace = new FTrace();
 	Trace->Name = Name;
 	Trace->Id = QuickStoreHash(Name);
+	Trace->Handle = UPTRINT(Handle);
 
 	Traces.Add(Trace);
 	return Trace;
diff --git a/Engine/Source/Developer/TraceAnalysis/Private/Store/AsioStore.h b/Engine/Source/Developer/TraceAnalysis/Private/Store/AsioStore.h
--- a/Engine/Source/Developer/TraceAnalysis/Private/Store/AsioStore.h
+++ b/Engine/Source/Developer/TraceAnalysis/Private/Store/AsioStore.h
@@ -39,6 +39,7 @@ public:
 	private:
 		friend			FAsioStore;
 		FString			Name;
+		UPTRINT			Handle;
 		uint32			Id;
 	};
 
