Implemented FTraceInfo::GetName()

#rnx

diff --git a/Engine/Source/Developer/TraceAnalysis/Private/Store/AsioStoreCborServer.cpp b/Engine/Source/Developer/TraceAnalysis/Private/Store/AsioStoreCborServer.cpp
--- a/Engine/Source/Developer/TraceAnalysis/Private/Store/AsioStoreCborServer.cpp
+++ b/Engine/Source/Developer/TraceAnalysis/Private/Store/AsioStoreCborServer.cpp
@@ -125,9 +125,22 @@ void FAsioStoreCborPeer::OnTraceInfo()
 		return SendError(EStatusCode::BadRequest);
 	}
 
+	const TCHAR* Name = Trace->GetName();
+	char OutName[128];
+	for (char& Out : OutName)
+	{
+		Out = char(*Name++);
+		if (Out == '\0')
+		{
+			break;
+		}
+	}
+	OutName[sizeof(OutName) - 1] = '\0';
+
 	TPayloadBuilder<> Builder((int32)EStatusCode::Success);
 	Builder.AddInteger("id", Trace->GetId());
 	Builder.AddInteger("size", Trace->GetSize());
+	Builder.AddString("name", OutName);
 	SendResponse(Builder.Done());
 }
 
diff --git a/Engine/Source/Developer/TraceAnalysis/Private/Store/StoreClient.cpp b/Engine/Source/Developer/TraceAnalysis/Private/Store/StoreClient.cpp
--- a/Engine/Source/Developer/TraceAnalysis/Private/Store/StoreClient.cpp
+++ b/Engine/Source/Developer/TraceAnalysis/Private/Store/StoreClient.cpp
@@ -267,6 +267,13 @@ uint64 FStoreClient::FTraceInfo::GetSize() const
 	return Response->GetInteger("size", 0);
 }
 
+////////////////////////////////////////////////////////////////////////////////
+FAnsiStringView FStoreClient::FTraceInfo::GetName() const
+{
+	const auto* Response = (const FResponse*)this;
+	return Response->GetString("name", "nameless");
+}
+
 
 
 ////////////////////////////////////////////////////////////////////////////////
diff --git a/Engine/Source/Developer/TraceAnalysis/Public/Trace/StoreClient.h b/Engine/Source/Developer/TraceAnalysis/Public/Trace/StoreClient.h
--- a/Engine/Source/Developer/TraceAnalysis/Public/Trace/StoreClient.h
+++ b/Engine/Source/Developer/TraceAnalysis/Public/Trace/StoreClient.h
@@ -2,6 +2,7 @@
 
 #pragma once
 
+#include "Containers/StringView.h"
 #include "HAL/Platform.h"
 
 namespace Trace
@@ -20,6 +21,7 @@ public:
 
 	struct TRACEANALYSIS_API FTraceInfo
 	{
+		FAnsiStringView	GetName() const;
 		uint32			GetId() const;
 		uint64			GetSize() const;
 		//const TCHAR*	GetMetadata(const TCHAR* Key) const;
