Correct implementation of FStoreClient::IsValid().

#rnx

diff --git a/Engine/Source/Developer/TraceAnalysis/Private/Store/StoreClient.cpp b/Engine/Source/Developer/TraceAnalysis/Private/Store/StoreClient.cpp
--- a/Engine/Source/Developer/TraceAnalysis/Private/Store/StoreClient.cpp
+++ b/Engine/Source/Developer/TraceAnalysis/Private/Store/StoreClient.cpp
@@ -68,6 +68,7 @@ class FStoreCborClient
 public:
 							FStoreCborClient();
 							~FStoreCborClient();
+	bool					IsValid() const;
 	const FResponse&		GetResponse() const;
 	bool					Connect(const TCHAR* Host, uint16 Port);
 	bool					GetStatus();
@@ -93,6 +94,12 @@ FStoreCborClient::~FStoreCborClient()
 {
 }
 
+////////////////////////////////////////////////////////////////////////////////
+bool FStoreCborClient::IsValid() const
+{
+	return Socket.is_open();
+}
+
 ////////////////////////////////////////////////////////////////////////////////
 const FResponse& FStoreCborClient::GetResponse() const
 {
@@ -310,7 +317,7 @@ FStoreClient::~FStoreClient()
 ////////////////////////////////////////////////////////////////////////////////
 bool FStoreClient::IsValid() const
 {
-	return Impl != nullptr;
+	return Impl->IsValid();
 }
 
 ////////////////////////////////////////////////////////////////////////////////
