Ununity fix.

#rnx

diff --git a/Engine/Source/Developer/TraceAnalysis/Public/Trace/StoreClient.h b/Engine/Source/Developer/TraceAnalysis/Public/Trace/StoreClient.h
--- a/Engine/Source/Developer/TraceAnalysis/Public/Trace/StoreClient.h
+++ b/Engine/Source/Developer/TraceAnalysis/Public/Trace/StoreClient.h
@@ -4,6 +4,7 @@
 
 #include "Containers/StringView.h"
 #include "HAL/Platform.h"
+#include "Templates/UniquePtr.h"
 
 namespace Trace
 {
