Added a FStoreClient::GetTraceInfoById().

#rnx

diff --git a/Engine/Source/Developer/TraceAnalysis/Private/Store/AsioStore.cpp b/Engine/Source/Developer/TraceAnalysis/Private/Store/AsioStore.cpp
--- a/Engine/Source/Developer/TraceAnalysis/Private/Store/AsioStore.cpp
+++ b/Engine/Source/Developer/TraceAnalysis/Private/Store/AsioStore.cpp
@@ -89,6 +89,20 @@ const FAsioStore::FTrace* FAsioStore::GetTraceInfo(uint32 Index) const
 	return Traces[Index];
 }
 
+////////////////////////////////////////////////////////////////////////////////
+const FAsioStore::FTrace* FAsioStore::GetTraceInfoById(uint32 Id) const
+{
+	for (FTrace* Trace : Traces)
+	{
+		if (Trace->GetId() == Id)
+		{
+			return Trace;
+		}
+	}
+
+	return nullptr;
+}
+
 ////////////////////////////////////////////////////////////////////////////////
 FAsioStore::FTrace* FAsioStore::GetTrace(uint32 Id)
 {
diff --git a/Engine/Source/Developer/TraceAnalysis/Private/Store/AsioStore.h b/Engine/Source/Developer/TraceAnalysis/Private/Store/AsioStore.h
--- a/Engine/Source/Developer/TraceAnalysis/Private/Store/AsioStore.h
+++ b/Engine/Source/Developer/TraceAnalysis/Private/Store/AsioStore.h
@@ -53,6 +53,7 @@ public:
 						~FAsioStore();
 	uint32				GetTraceCount() const;
 	const FTrace*		GetTraceInfo(uint32 Index) const;
+	const FTrace*		GetTraceInfoById(uint32 Id) const;
 	int32				ReadTrace(uint32 Id);
 	FNewTrace			CreateTrace();
 	FAsioReadable*		OpenTrace(uint32 Id);
diff --git a/Engine/Source/Developer/TraceAnalysis/Private/Store/AsioStoreCborServer.cpp b/Engine/Source/Developer/TraceAnalysis/Private/Store/AsioStoreCborServer.cpp
--- a/Engine/Source/Developer/TraceAnalysis/Private/Store/AsioStoreCborServer.cpp
+++ b/Engine/Source/Developer/TraceAnalysis/Private/Store/AsioStoreCborServer.cpp
@@ -113,13 +113,27 @@ void FAsioStoreCborPeer::OnTraceCount()
 ////////////////////////////////////////////////////////////////////////////////
 void FAsioStoreCborPeer::OnTraceInfo()
 {
+	const FAsioStore::FTrace* Trace = nullptr;
+
 	int32 Index = int32(Response.GetInteger("index", -1));
-	if (Index < 0)
+	if (Index >= 0)
+	{
+		Trace = Store.GetTraceInfo(Index);
+	}
+	else
+	{
+		uint32 Id = uint32(Response.GetInteger("id", 0));
+		if (Id != 0)
+		{
+			Trace = Store.GetTraceInfoById(Id);
+		}
+	}
+
+	if (Trace == nullptr)
 	{
 		return SendError(EStatusCode::BadRequest);
 	}
 
-	const FAsioStore::FTrace* Trace = Store.GetTraceInfo(Index);
 	if (Trace == nullptr)
 	{
 		return SendError(EStatusCode::BadRequest);
diff --git a/Engine/Source/Developer/TraceAnalysis/Private/Store/StoreClient.cpp b/Engine/Source/Developer/TraceAnalysis/Private/Store/StoreClient.cpp
--- a/Engine/Source/Developer/TraceAnalysis/Private/Store/StoreClient.cpp
+++ b/Engine/Source/Developer/TraceAnalysis/Private/Store/StoreClient.cpp
@@ -74,6 +74,7 @@ public:
 	bool					GetStatus();
 	bool					GetTraceCount();
 	bool					GetTraceInfo(uint32 Index);
+	bool					GetTraceInfoById(uint32 Id);
 	FTraceDataStream*		ReadTrace(uint32 Id);
 
 private:
@@ -201,6 +202,15 @@ bool FStoreCborClient::GetTraceInfo(uint32 Index)
 	return Communicate(Payload);
 }
 
+////////////////////////////////////////////////////////////////////////////////
+bool FStoreCborClient::GetTraceInfoById(uint32 Id)
+{
+	TPayloadBuilder<> Builder("trace/info");
+	Builder.AddInteger("id", int32(Id));
+	FPayload Payload = Builder.Done();
+	return Communicate(Payload);
+}
+
 ////////////////////////////////////////////////////////////////////////////////
 FTraceDataStream* FStoreCborClient::ReadTrace(uint32 Id)
 {
@@ -367,6 +377,22 @@ const FStoreClient::FTraceInfo* FStoreClient::GetTraceInfo(uint32 Index)
 #endif // TRACE_WITH_ASIO
 }
 
+////////////////////////////////////////////////////////////////////////////////
+const FStoreClient::FTraceInfo* FStoreClient::GetTraceInfoById(uint32 Id)
+{
+#if TRACE_WITH_ASIO
+	if (!Impl->GetTraceInfoById(Id))
+	{
+		return nullptr;
+	}
+
+	const FResponse& Response = Impl->GetResponse();
+	return (FTraceInfo*)(&Response);
+#else
+	return nullptr;
+#endif // TRACE_WITH_ASIO
+}
+
 ////////////////////////////////////////////////////////////////////////////////
 FStoreClient::FTraceData FStoreClient::ReadTrace(uint32 Id)
 {
diff --git a/Engine/Source/Developer/TraceAnalysis/Public/Trace/StoreClient.h b/Engine/Source/Developer/TraceAnalysis/Public/Trace/StoreClient.h
--- a/Engine/Source/Developer/TraceAnalysis/Public/Trace/StoreClient.h
+++ b/Engine/Source/Developer/TraceAnalysis/Public/Trace/StoreClient.h
@@ -41,6 +41,7 @@ public:
 	const FStatus*		GetStatus();
 	uint32				GetTraceCount();
 	const FTraceInfo*	GetTraceInfo(uint32 Index);
+	const FTraceInfo*	GetTraceInfoById(uint32 Id);
 	FTraceData			ReadTrace(uint32 Id);
 
 #if 0
