I forgot the ASIO define.

#rnx

diff --git a/Engine/Source/Developer/TraceAnalysis/Private/Store/StoreClient.cpp b/Engine/Source/Developer/TraceAnalysis/Private/Store/StoreClient.cpp
--- a/Engine/Source/Developer/TraceAnalysis/Private/Store/StoreClient.cpp
+++ b/Engine/Source/Developer/TraceAnalysis/Private/Store/StoreClient.cpp
@@ -327,7 +327,11 @@ FStoreClient::~FStoreClient()
 ////////////////////////////////////////////////////////////////////////////////
 bool FStoreClient::IsValid() const
 {
+#if TRACE_WITH_ASIO
 	return Impl->IsValid();
+#else
+	return false;
+#endif
 }
 
 ////////////////////////////////////////////////////////////////////////////////
