One missed ASIO define. One missed include.

#rnx

diff --git a/Engine/Source/Developer/TraceAnalysis/Private/Store/StoreClient.cpp b/Engine/Source/Developer/TraceAnalysis/Private/Store/StoreClient.cpp
--- a/Engine/Source/Developer/TraceAnalysis/Private/Store/StoreClient.cpp
+++ b/Engine/Source/Developer/TraceAnalysis/Private/Store/StoreClient.cpp
@@ -8,7 +8,6 @@
 #include "AsioStore.h"
 #include "CborPayload.h"
 #include "Templates/UnrealTemplate.h"
-#include "Trace/DataStream.h"
 
 namespace Trace
 {
@@ -400,7 +399,11 @@ const FStoreClient::FTraceInfo* FStoreClient::GetTraceInfoById(uint32 Id)
 ////////////////////////////////////////////////////////////////////////////////
 FStoreClient::FTraceData FStoreClient::ReadTrace(uint32 Id)
 {
+#if TRACE_WITH_ASIO
 	return FTraceData(Impl->ReadTrace(Id));
+#else
+	return nullptr;
+#endif // TRACE_WITH_ASIO
 }
 
 } // namespace Trace
diff --git a/Engine/Source/Developer/TraceAnalysis/Public/Trace/StoreClient.h b/Engine/Source/Developer/TraceAnalysis/Public/Trace/StoreClient.h
--- a/Engine/Source/Developer/TraceAnalysis/Public/Trace/StoreClient.h
+++ b/Engine/Source/Developer/TraceAnalysis/Public/Trace/StoreClient.h
@@ -5,6 +5,7 @@
 #include "Containers/StringView.h"
 #include "HAL/Platform.h"
 #include "Templates/UniquePtr.h"
+#include "Trace/DataStream.h"
 
 namespace Trace
 {
