Added accessors to get a socket's ports and remote ip address.

#rnx

diff --git a/Engine/Source/Developer/TraceAnalysis/Private/Store/AsioSocket.cpp b/Engine/Source/Developer/TraceAnalysis/Private/Store/AsioSocket.cpp
--- a/Engine/Source/Developer/TraceAnalysis/Private/Store/AsioSocket.cpp
+++ b/Engine/Source/Developer/TraceAnalysis/Private/Store/AsioSocket.cpp
@@ -26,6 +26,43 @@ bool FAsioSocket::IsOpen() const
 	return Socket.is_open();
 }
 
+////////////////////////////////////////////////////////////////////////////////
+uint32 FAsioSocket::GetRemoteAddress() const
+{
+	if (!IsOpen())
+	{
+		return 0;
+	}
+
+	const asio::ip::tcp::endpoint Endpoint = Socket.remote_endpoint();
+	asio::ip::address Address = Endpoint.address();
+	return Address.is_v4() ? Address.to_v4().to_uint() : 0;
+}
+
+////////////////////////////////////////////////////////////////////////////////
+uint32 FAsioSocket::GetRemotePort() const
+{
+	if (!IsOpen())
+	{
+		return 0;
+	}
+
+	const asio::ip::tcp::endpoint Endpoint = Socket.remote_endpoint();
+	return Endpoint.port();
+}
+
+////////////////////////////////////////////////////////////////////////////////
+uint32 FAsioSocket::GetLocalPort() const
+{
+	if (!IsOpen())
+	{
+		return 0;
+	}
+
+	const asio::ip::tcp::endpoint Endpoint = Socket.local_endpoint();
+	return Endpoint.port();
+}
+
 ////////////////////////////////////////////////////////////////////////////////
 void FAsioSocket::Close()
 {
diff --git a/Engine/Source/Developer/TraceAnalysis/Private/Store/AsioSocket.h b/Engine/Source/Developer/TraceAnalysis/Private/Store/AsioSocket.h
--- a/Engine/Source/Developer/TraceAnalysis/Private/Store/AsioSocket.h
+++ b/Engine/Source/Developer/TraceAnalysis/Private/Store/AsioSocket.h
@@ -20,6 +20,9 @@ public:
 							FAsioSocket(asio::ip::tcp::socket& InSocket);
 	virtual					~FAsioSocket();
 	bool					IsOpen() const;
+	uint32					GetRemoteAddress() const;
+	uint32					GetRemotePort() const;
+	uint32					GetLocalPort() const;
 	virtual void			Close() override;
 	virtual bool			Read(void* Dest, uint32 Size, FAsioIoSink* Sink, uint32 Id) override;
 	virtual bool			ReadSome(void* Dest, uint32 DestSize, FAsioIoSink* Sink, uint32 Id) override;
