Make information about a recorder's sessions available to the outside world.

#rnx

diff --git a/Engine/Source/Developer/TraceAnalysis/Private/Store/AsioRecorder.cpp b/Engine/Source/Developer/TraceAnalysis/Private/Store/AsioRecorder.cpp
--- a/Engine/Source/Developer/TraceAnalysis/Private/Store/AsioRecorder.cpp
+++ b/Engine/Source/Developer/TraceAnalysis/Private/Store/AsioRecorder.cpp
@@ -74,6 +74,14 @@ void FAsioRecorderRelay::OnIoComplete(uint32 Id, int32 Size)
 
 
 
+////////////////////////////////////////////////////////////////////////////////
+uint32 FAsioRecorder::FSession::GetTraceId() const
+{
+	return TraceId;
+}
+
+
+
 ////////////////////////////////////////////////////////////////////////////////
 FAsioRecorder::FAsioRecorder(asio::io_context& IoContext, FAsioStore& InStore)
 : FAsioTcpServer(IoContext)
@@ -93,6 +101,23 @@ FAsioRecorder::~FAsioRecorder()
 	}
 }
 
+////////////////////////////////////////////////////////////////////////////////
+uint32 FAsioRecorder::GetSessionCount() const
+{
+	return uint32(Sessions.Num());
+}
+
+////////////////////////////////////////////////////////////////////////////////
+const FAsioRecorder::FSession* FAsioRecorder::GetSessionInfo(uint32 Index) const
+{
+	if (Index >= uint32(Sessions.Num()))
+	{
+		return nullptr;
+	}
+
+	return Sessions.GetData() + Index;
+}
+
 ////////////////////////////////////////////////////////////////////////////////
 bool FAsioRecorder::OnAccept(asio::ip::tcp::socket& Socket)
 {
@@ -102,8 +127,13 @@ bool FAsioRecorder::OnAccept(asio::ip::tcp::socket& Socket)
 		return true;
 	}
 
-	FAsioRecorderRelay* Relay = new FAsioRecorderRelay(Socket, Trace.Writeable);
-	Sessions.Add({Relay, Trace.Id});
+	auto* Relay = new FAsioRecorderRelay(Socket, Trace.Writeable);
+
+	FSession Session;
+	Session.Relay = Relay;
+	Session.TraceId = Trace.Id;
+	Sessions.Add(Session);
+
 	return true;
 }
 
diff --git a/Engine/Source/Developer/TraceAnalysis/Private/Store/AsioRecorder.h b/Engine/Source/Developer/TraceAnalysis/Private/Store/AsioRecorder.h
--- a/Engine/Source/Developer/TraceAnalysis/Private/Store/AsioRecorder.h
+++ b/Engine/Source/Developer/TraceAnalysis/Private/Store/AsioRecorder.h
@@ -22,16 +22,23 @@ class FAsioRecorder
 	, public FAsioTickable
 {
 public:
-								FAsioRecorder(asio::io_context& IoContext, FAsioStore& InStore);
-								~FAsioRecorder();
-
-private:
-	struct FSession
+	class FSession
 	{
+	public:
+		uint32					GetTraceId() const;
+
+	private:
+		friend					FAsioRecorder;
 		FAsioRecorderRelay*		Relay;
 		uint32					TraceId;
 	};
 
+								FAsioRecorder(asio::io_context& IoContext, FAsioStore& InStore);
+								~FAsioRecorder();
+	uint32						GetSessionCount() const;
+	const FSession*				GetSessionInfo(uint32 Index) const;
+
+private:
 	virtual bool				OnAccept(asio::ip::tcp::socket& Socket) override;
 	virtual void				OnTick() override;
 	TArray<FSession>			Sessions;
