Paired down the store API so as to not get too carried away just yet.

#rnx

diff --git a/Engine/Source/Developer/TraceAnalysis/Private/Store/AsioStore.cpp b/Engine/Source/Developer/TraceAnalysis/Private/Store/AsioStore.cpp
--- a/Engine/Source/Developer/TraceAnalysis/Private/Store/AsioStore.cpp
+++ b/Engine/Source/Developer/TraceAnalysis/Private/Store/AsioStore.cpp
@@ -89,20 +89,6 @@ const FAsioStore::FTrace* FAsioStore::GetTraceInfo(uint32 Index) const
 	return Traces[Index];
 }
 
-////////////////////////////////////////////////////////////////////////////////
-const FAsioStore::FTrace* FAsioStore::GetTraceInfoById(uint32 Id) const
-{
-	for (FTrace* Trace : Traces)
-	{
-		if (Trace->GetId() == Id)
-		{
-			return Trace;
-		}
-	}
-
-	return nullptr;
-}
-
 ////////////////////////////////////////////////////////////////////////////////
 FAsioStore::FTrace* FAsioStore::GetTrace(uint32 Id)
 {
diff --git a/Engine/Source/Developer/TraceAnalysis/Private/Store/AsioStore.h b/Engine/Source/Developer/TraceAnalysis/Private/Store/AsioStore.h
--- a/Engine/Source/Developer/TraceAnalysis/Private/Store/AsioStore.h
+++ b/Engine/Source/Developer/TraceAnalysis/Private/Store/AsioStore.h
@@ -53,7 +53,6 @@ public:
 						~FAsioStore();
 	uint32				GetTraceCount() const;
 	const FTrace*		GetTraceInfo(uint32 Index) const;
-	const FTrace*		GetTraceInfoById(uint32 Id) const;
 	int32				ReadTrace(uint32 Id);
 	FNewTrace			CreateTrace();
 	FAsioReadable*		OpenTrace(uint32 Id);
diff --git a/Engine/Source/Developer/TraceAnalysis/Private/Store/AsioStoreCborServer.cpp b/Engine/Source/Developer/TraceAnalysis/Private/Store/AsioStoreCborServer.cpp
--- a/Engine/Source/Developer/TraceAnalysis/Private/Store/AsioStoreCborServer.cpp
+++ b/Engine/Source/Developer/TraceAnalysis/Private/Store/AsioStoreCborServer.cpp
@@ -125,7 +125,15 @@ void FAsioStoreCborPeer::OnTraceInfo()
 		uint32 Id = uint32(Response.GetInteger("id", 0));
 		if (Id != 0)
 		{
-			Trace = Store.GetTraceInfoById(Id);
+			for (int i = 0, n = Store.GetTraceCount(); i < n; ++i)
+			{
+				const FAsioStore::FTrace* Candidate = Store.GetTraceInfo(i);
+				if (Candidate->GetId() == Id)
+				{
+					Trace = Candidate;
+					break;
+				}
+			}
 		}
 	}
 
