Access to a recorder sessions remote IP address.

#rnx

diff --git a/Engine/Source/Developer/TraceAnalysis/Private/Store/AsioRecorder.cpp b/Engine/Source/Developer/TraceAnalysis/Private/Store/AsioRecorder.cpp
--- a/Engine/Source/Developer/TraceAnalysis/Private/Store/AsioRecorder.cpp
+++ b/Engine/Source/Developer/TraceAnalysis/Private/Store/AsioRecorder.cpp
@@ -19,6 +19,8 @@ public:
 						FAsioRecorderRelay(asio::ip::tcp::socket& Socket, FAsioWriteable* InOutput);
 	virtual				~FAsioRecorderRelay();
 	void				Close();
+	uint32				GetIpAddress() const;
+
 
 private:
 	virtual void		OnIoComplete(uint32 Id, int32 Size) override;
@@ -50,6 +52,12 @@ void FAsioRecorderRelay::Close()
 	Output->Close();
 }
 
+////////////////////////////////////////////////////////////////////////////////
+uint32 FAsioRecorderRelay::GetIpAddress() const
+{
+	return Input.GetRemoteAddress();
+}
+
 ////////////////////////////////////////////////////////////////////////////////
 void FAsioRecorderRelay::OnIoComplete(uint32 Id, int32 Size)
 {
@@ -80,6 +88,12 @@ uint32 FAsioRecorder::FSession::GetTraceId() const
 	return TraceId;
 }
 
+////////////////////////////////////////////////////////////////////////////////
+uint32 FAsioRecorder::FSession::GetIpAddress() const
+{
+	return Relay->GetIpAddress();
+}
+
 
 
 ////////////////////////////////////////////////////////////////////////////////
diff --git a/Engine/Source/Developer/TraceAnalysis/Private/Store/AsioRecorder.h b/Engine/Source/Developer/TraceAnalysis/Private/Store/AsioRecorder.h
--- a/Engine/Source/Developer/TraceAnalysis/Private/Store/AsioRecorder.h
+++ b/Engine/Source/Developer/TraceAnalysis/Private/Store/AsioRecorder.h
@@ -26,6 +26,7 @@ public:
 	{
 	public:
 		uint32					GetTraceId() const;
+		uint32					GetIpAddress() const;
 
 	private:
 		friend					FAsioRecorder;
