Give each of a recorder's sessions a unique identifier.

#rnx

diff --git a/Engine/Source/Developer/TraceAnalysis/Private/Store/AsioRecorder.cpp b/Engine/Source/Developer/TraceAnalysis/Private/Store/AsioRecorder.cpp
--- a/Engine/Source/Developer/TraceAnalysis/Private/Store/AsioRecorder.cpp
+++ b/Engine/Source/Developer/TraceAnalysis/Private/Store/AsioRecorder.cpp
@@ -82,6 +82,12 @@ void FAsioRecorderRelay::OnIoComplete(uint32 Id, int32 Size)
 
 
 
+////////////////////////////////////////////////////////////////////////////////
+uint32 FAsioRecorder::FSession::GetId() const
+{
+	return Id;
+}
+
 ////////////////////////////////////////////////////////////////////////////////
 uint32 FAsioRecorder::FSession::GetTraceId() const
 {
@@ -143,8 +149,16 @@ bool FAsioRecorder::OnAccept(asio::ip::tcp::socket& Socket)
 
 	auto* Relay = new FAsioRecorderRelay(Socket, Trace.Writeable);
 
+	uint32 IdPieces[] = {
+		Relay->GetIpAddress(),
+		Socket.remote_endpoint().port(),
+		Socket.local_endpoint().port(),
+		0,
+	};
+
 	FSession Session;
 	Session.Relay = Relay;
+	Session.Id = QuickStoreHash(IdPieces);
 	Session.TraceId = Trace.Id;
 	Sessions.Add(Session);
 
diff --git a/Engine/Source/Developer/TraceAnalysis/Private/Store/AsioRecorder.h b/Engine/Source/Developer/TraceAnalysis/Private/Store/AsioRecorder.h
--- a/Engine/Source/Developer/TraceAnalysis/Private/Store/AsioRecorder.h
+++ b/Engine/Source/Developer/TraceAnalysis/Private/Store/AsioRecorder.h
@@ -25,12 +25,14 @@ public:
 	class FSession
 	{
 	public:
+		uint32					GetId() const;
 		uint32					GetTraceId() const;
 		uint32					GetIpAddress() const;
 
 	private:
 		friend					FAsioRecorder;
 		FAsioRecorderRelay*		Relay;
+		uint32					Id;
 		uint32					TraceId;
 	};
 
