Added an API for a client to query a store's active sessions

#rnx

diff --git a/Engine/Source/Developer/TraceAnalysis/Private/Store/AsioStoreCborServer.cpp b/Engine/Source/Developer/TraceAnalysis/Private/Store/AsioStoreCborServer.cpp
--- a/Engine/Source/Developer/TraceAnalysis/Private/Store/AsioStoreCborServer.cpp
+++ b/Engine/Source/Developer/TraceAnalysis/Private/Store/AsioStoreCborServer.cpp
@@ -37,6 +37,8 @@ protected:
 
 	void			OnPayload();
 	void			OnConnect();
+	void			OnSessionCount();
+	void			OnSessionInfo();
 	void			OnStatus();
 	void			OnTraceCount();
 	void			OnTraceInfo();
@@ -94,6 +96,61 @@ void FAsioStoreCborPeer::OnConnect()
 	SendResponse(Builder.Done());
 }
 
+////////////////////////////////////////////////////////////////////////////////
+void FAsioStoreCborPeer::OnSessionCount()
+{
+	TPayloadBuilder<> Builder((int32)EStatusCode::Success);
+	Builder.AddInteger("count", Recorder.GetSessionCount());
+	SendResponse(Builder.Done());
+}
+
+////////////////////////////////////////////////////////////////////////////////
+void FAsioStoreCborPeer::OnSessionInfo()
+{
+	const FAsioRecorder::FSession* Session = nullptr;
+
+	int32 Index = int32(Response.GetInteger("index", -1));
+	if (Index >= 0)
+	{
+		Session = Recorder.GetSessionInfo(Index);
+	}
+	else if (uint32 Id = uint32(Response.GetInteger("id", 0)))
+	{
+		for (int i = 0, n = Recorder.GetSessionCount(); i < n; ++i)
+		{
+			const FAsioRecorder::FSession* Candidate = Recorder.GetSessionInfo(i);
+			if (Candidate->GetId() == Id)
+			{
+				Session = Candidate;
+				break;
+			}
+		}
+	}
+	else if (uint32 TraceId = uint32(Response.GetInteger("trace_id", 0)))
+	{
+		for (int i = 0, n = Recorder.GetSessionCount(); i < n; ++i)
+		{
+			const FAsioRecorder::FSession* Candidate = Recorder.GetSessionInfo(i);
+			if (Candidate->GetTraceId() == TraceId)
+			{
+				Session = Candidate;
+				break;
+			}
+		}
+	}
+
+	if (Session == nullptr)
+	{
+		return SendError(EStatusCode::BadRequest);
+	}
+
+	TPayloadBuilder<> Builder((int32)EStatusCode::Success);
+	Builder.AddInteger("id", Session->GetId());
+	Builder.AddInteger("trace_id", Session->GetTraceId());
+	Builder.AddInteger("ip_address", Session->GetIpAddress());
+	SendResponse(Builder.Done());
+}
+
 ////////////////////////////////////////////////////////////////////////////////
 void FAsioStoreCborPeer::OnStatus()
 {
@@ -208,11 +265,13 @@ void FAsioStoreCborPeer::OnPayload()
 		uint32	Hash;
 		void	(FAsioStoreCborPeer::*Func)();
 	} const DispatchTable[] = {
-		{ QuickStoreHash("connect"),	&FAsioStoreCborPeer::OnConnect },
-		{ QuickStoreHash("status"),		&FAsioStoreCborPeer::OnStatus },
-		{ QuickStoreHash("trace/count"),&FAsioStoreCborPeer::OnTraceCount },
-		{ QuickStoreHash("trace/info"),	&FAsioStoreCborPeer::OnTraceInfo },
-		{ QuickStoreHash("trace/read"),	&FAsioStoreCborPeer::OnTraceRead },
+		{ QuickStoreHash("connect"),		&FAsioStoreCborPeer::OnConnect },
+		{ QuickStoreHash("session/count"),	&FAsioStoreCborPeer::OnSessionCount },
+		{ QuickStoreHash("session/info"),	&FAsioStoreCborPeer::OnSessionInfo },
+		{ QuickStoreHash("status"),			&FAsioStoreCborPeer::OnStatus },
+		{ QuickStoreHash("trace/count"),	&FAsioStoreCborPeer::OnTraceCount },
+		{ QuickStoreHash("trace/info"),		&FAsioStoreCborPeer::OnTraceInfo },
+		{ QuickStoreHash("trace/read"),		&FAsioStoreCborPeer::OnTraceRead },
 	};
 
 	uint32 MethodHash = QuickStoreHash(Method);
diff --git a/Engine/Source/Developer/TraceAnalysis/Private/Store/StoreClient.cpp b/Engine/Source/Developer/TraceAnalysis/Private/Store/StoreClient.cpp
--- a/Engine/Source/Developer/TraceAnalysis/Private/Store/StoreClient.cpp
+++ b/Engine/Source/Developer/TraceAnalysis/Private/Store/StoreClient.cpp
@@ -75,6 +75,10 @@ public:
 	bool					GetTraceInfo(uint32 Index);
 	bool					GetTraceInfoById(uint32 Id);
 	FTraceDataStream*		ReadTrace(uint32 Id);
+	bool					GetSessionCount();
+	bool					GetSessionInfo(uint32 Index);
+	bool					GetSessionInfoById(uint32 Id);
+	bool					GetSessionInfoByTraceId(uint32 TraceId);
 
 private:
 	bool					Communicate(const FPayload& Payload);
@@ -241,6 +245,41 @@ FTraceDataStream* FStoreCborClient::ReadTrace(uint32 Id)
 	return new FTraceDataStream(SenderSocket);
 }
 
+////////////////////////////////////////////////////////////////////////////////
+bool FStoreCborClient::GetSessionCount()
+{
+	TPayloadBuilder<> Builder("session/count");
+	FPayload Payload = Builder.Done();
+	return Communicate(Payload);
+}
+
+////////////////////////////////////////////////////////////////////////////////
+bool FStoreCborClient::GetSessionInfo(uint32 Index)
+{
+	TPayloadBuilder<> Builder("session/info");
+	Builder.AddInteger("index", int32(Index));
+	FPayload Payload = Builder.Done();
+	return Communicate(Payload);
+}
+
+////////////////////////////////////////////////////////////////////////////////
+bool FStoreCborClient::GetSessionInfoById(uint32 Id)
+{
+	TPayloadBuilder<> Builder("session/info");
+	Builder.AddInteger("id", int32(Id));
+	FPayload Payload = Builder.Done();
+	return Communicate(Payload);
+}
+
+////////////////////////////////////////////////////////////////////////////////
+bool FStoreCborClient::GetSessionInfoByTraceId(uint32 TraceId)
+{
+	TPayloadBuilder<> Builder("session/info");
+	Builder.AddInteger("trace_id", int32(TraceId));
+	FPayload Payload = Builder.Done();
+	return Communicate(Payload);
+}
+
 } // namespace Trace
 
 #endif // TRACE_WITH_ASIO
@@ -292,6 +331,29 @@ FAnsiStringView FStoreClient::FTraceInfo::GetName() const
 
 
 
+////////////////////////////////////////////////////////////////////////////////
+uint32 FStoreClient::FSessionInfo::GetId() const
+{
+	const auto* Response = (const FResponse*)this;
+	return Response->GetInteger("id", 0);
+}
+
+////////////////////////////////////////////////////////////////////////////////
+uint32 FStoreClient::FSessionInfo::GetTraceId() const
+{
+	const auto* Response = (const FResponse*)this;
+	return Response->GetInteger("trace_id", 0);
+}
+
+////////////////////////////////////////////////////////////////////////////////
+uint32 FStoreClient::FSessionInfo::GetIpAddress() const
+{
+	const auto* Response = (const FResponse*)this;
+	return Response->GetInteger("ip_address", 0);
+}
+
+
+
 ////////////////////////////////////////////////////////////////////////////////
 struct FStoreClient::FImpl
 #if TRACE_WITH_ASIO
@@ -406,4 +468,67 @@ FStoreClient::FTraceData FStoreClient::ReadTrace(uint32 Id)
 #endif // TRACE_WITH_ASIO
 }
 
+////////////////////////////////////////////////////////////////////////////////
+uint32 FStoreClient::GetSessionCount() const
+{
+#if TRACE_WITH_ASIO
+	if (!Impl->GetSessionCount())
+	{
+		return 0;
+	}
+
+	return Impl->GetResponse().GetInteger("count", 0);
+#else
+	return 0;
+#endif // TRACE_WITH_ASIO
+}
+
+////////////////////////////////////////////////////////////////////////////////
+const FStoreClient::FSessionInfo* FStoreClient::GetSessionInfo(uint32 Index) const
+{
+#if TRACE_WITH_ASIO
+	if (!Impl->GetSessionInfo(Index))
+	{
+		return nullptr;
+	}
+
+	const FResponse& Response = Impl->GetResponse();
+	return (FSessionInfo*)(&Response);
+#else
+	return nullptr;
+#endif // TRACE_WITH_ASIO
+}
+
+////////////////////////////////////////////////////////////////////////////////
+const FStoreClient::FSessionInfo* FStoreClient::GetSessionInfoById(uint32 Id) const
+{
+#if TRACE_WITH_ASIO
+	if (!Impl->GetSessionInfoById(Id))
+	{
+		return nullptr;
+	}
+
+	const FResponse& Response = Impl->GetResponse();
+	return (FSessionInfo*)(&Response);
+#else
+	return nullptr;
+#endif // TRACE_WITH_ASIO
+}
+
+////////////////////////////////////////////////////////////////////////////////
+const FStoreClient::FSessionInfo* FStoreClient::GetSessionInfoByTraceId(uint32 TraceId) const
+{
+#if TRACE_WITH_ASIO
+	if (!Impl->GetSessionInfoByTraceId(TraceId))
+	{
+		return nullptr;
+	}
+
+	const FResponse& Response = Impl->GetResponse();
+	return (FSessionInfo*)(&Response);
+#else
+	return nullptr;
+#endif // TRACE_WITH_ASIO
+}
+
 } // namespace Trace
diff --git a/Engine/Source/Developer/TraceAnalysis/Public/Trace/StoreClient.h b/Engine/Source/Developer/TraceAnalysis/Public/Trace/StoreClient.h
--- a/Engine/Source/Developer/TraceAnalysis/Public/Trace/StoreClient.h
+++ b/Engine/Source/Developer/TraceAnalysis/Public/Trace/StoreClient.h
@@ -44,12 +44,10 @@ public:
 	const FTraceInfo*	GetTraceInfo(uint32 Index);
 	const FTraceInfo*	GetTraceInfoById(uint32 Id);
 	FTraceData			ReadTrace(uint32 Id);
-
 #if 0
 	template <typename Lambda> uint32 GetTraceInfos(uint32 StartIndex, uint32 Count, Lambda&& Callback) const;
 #endif // 0
 
-#if 0
 	struct TRACEANALYSIS_API FSessionInfo
 	{
 		uint32			GetId() const;
@@ -57,8 +55,9 @@ public:
 		uint32			GetIpAddress() const;
 	};
 	uint32				GetSessionCount() const;
-	const FSessionInfo* GetSessionInfo(uint32 Index, Lambda&& Callback) const;
-#endif // 0
+	const FSessionInfo* GetSessionInfo(uint32 Index) const;
+	const FSessionInfo* GetSessionInfoById(uint32 Id) const;
+	const FSessionInfo* GetSessionInfoByTraceId(uint32 TraceId) const;
 #if 0
 	template <typename Lambda> uint32	GetSessionInfos(uint32 StartIndex, uint32 Count, Lambda&& Callback) const;
 #endif // 0
