Ended up with a bit of duplication here somehow or other.

#rnx

diff --git a/Engine/Source/Developer/TraceAnalysis/Private/Store/AsioStoreCborServer.cpp b/Engine/Source/Developer/TraceAnalysis/Private/Store/AsioStoreCborServer.cpp
--- a/Engine/Source/Developer/TraceAnalysis/Private/Store/AsioStoreCborServer.cpp
+++ b/Engine/Source/Developer/TraceAnalysis/Private/Store/AsioStoreCborServer.cpp
@@ -199,11 +199,6 @@ void FAsioStoreCborPeer::OnTraceInfo()
 		return SendError(EStatusCode::BadRequest);
 	}
 
-	if (Trace == nullptr)
-	{
-		return SendError(EStatusCode::BadRequest);
-	}
-
 	const TCHAR* Name = Trace->GetName();
 	char OutName[128];
 	for (char& Out : OutName)
