Added access to the timestamps of a store's traces

#rnx

diff --git a/Engine/Source/Developer/TraceAnalysis/Private/Store/AsioStore.cpp b/Engine/Source/Developer/TraceAnalysis/Private/Store/AsioStore.cpp
--- a/Engine/Source/Developer/TraceAnalysis/Private/Store/AsioStore.cpp
+++ b/Engine/Source/Developer/TraceAnalysis/Private/Store/AsioStore.cpp
@@ -33,6 +33,18 @@ uint64 FAsioStore::FTrace::GetSize() const
 	return FileSize.QuadPart;
 }
 
+////////////////////////////////////////////////////////////////////////////////
+uint64 FAsioStore::FTrace::GetTimestamp() const
+{
+	HANDLE Inner = HANDLE(Handle);
+	FILETIME Time;
+	if (!GetFileTime(Inner, &Time, nullptr, nullptr))
+	{
+		return 0;
+	}
+	return (uint64(Time.dwHighDateTime) << 32ull) | uint64(Time.dwLowDateTime);
+}
+
 
 
 ////////////////////////////////////////////////////////////////////////////////
diff --git a/Engine/Source/Developer/TraceAnalysis/Private/Store/AsioStore.h b/Engine/Source/Developer/TraceAnalysis/Private/Store/AsioStore.h
--- a/Engine/Source/Developer/TraceAnalysis/Private/Store/AsioStore.h
+++ b/Engine/Source/Developer/TraceAnalysis/Private/Store/AsioStore.h
@@ -35,6 +35,7 @@ public:
 		const TCHAR*	GetName() const;
 		uint32			GetId() const;
 		uint64			GetSize() const;
+		uint64			GetTimestamp() const;
 
 	private:
 		friend			FAsioStore;
diff --git a/Engine/Source/Developer/TraceAnalysis/Private/Store/AsioStoreCborServer.cpp b/Engine/Source/Developer/TraceAnalysis/Private/Store/AsioStoreCborServer.cpp
--- a/Engine/Source/Developer/TraceAnalysis/Private/Store/AsioStoreCborServer.cpp
+++ b/Engine/Source/Developer/TraceAnalysis/Private/Store/AsioStoreCborServer.cpp
@@ -214,6 +214,7 @@ void FAsioStoreCborPeer::OnTraceInfo()
 	TPayloadBuilder<> Builder((int32)EStatusCode::Success);
 	Builder.AddInteger("id", Trace->GetId());
 	Builder.AddInteger("size", Trace->GetSize());
+	Builder.AddInteger("timestamp", Trace->GetTimestamp());
 	Builder.AddString("name", OutName);
 	SendResponse(Builder.Done());
 }
diff --git a/Engine/Source/Developer/TraceAnalysis/Private/Store/StoreClient.cpp b/Engine/Source/Developer/TraceAnalysis/Private/Store/StoreClient.cpp
--- a/Engine/Source/Developer/TraceAnalysis/Private/Store/StoreClient.cpp
+++ b/Engine/Source/Developer/TraceAnalysis/Private/Store/StoreClient.cpp
@@ -322,6 +322,13 @@ uint64 FStoreClient::FTraceInfo::GetSize() const
 	return Response->GetInteger("size", 0);
 }
 
+////////////////////////////////////////////////////////////////////////////////
+uint64 FStoreClient::FTraceInfo::GetTimestamp() const
+{
+	const auto* Response = (const FResponse*)this;
+	return Response->GetInteger("timestamp", 0);
+}
+
 ////////////////////////////////////////////////////////////////////////////////
 FAnsiStringView FStoreClient::FTraceInfo::GetName() const
 {
diff --git a/Engine/Source/Developer/TraceAnalysis/Public/Trace/StoreClient.h b/Engine/Source/Developer/TraceAnalysis/Public/Trace/StoreClient.h
--- a/Engine/Source/Developer/TraceAnalysis/Public/Trace/StoreClient.h
+++ b/Engine/Source/Developer/TraceAnalysis/Public/Trace/StoreClient.h
@@ -26,6 +26,7 @@ public:
 		FAnsiStringView	GetName() const;
 		uint32			GetId() const;
 		uint64			GetSize() const;
+		uint64			GetTimestamp() const;
 		//const TCHAR*	GetMetadata(const TCHAR* Key) const;
 		//template <typename Lambda> uint32 ReadMetadata(Lambda&& Callback) const;
 	};
