Always consider traces to be live.

#rnx

diff --git a/Engine/Source/Developer/TraceAnalysis/Private/DataStream.cpp b/Engine/Source/Developer/TraceAnalysis/Private/DataStream.cpp
--- a/Engine/Source/Developer/TraceAnalysis/Private/DataStream.cpp
+++ b/Engine/Source/Developer/TraceAnalysis/Private/DataStream.cpp
@@ -23,10 +23,18 @@ int32 FFileStream::Read(void* Data, uint32 Size)
 	{
 		return 0;
 	}
-	uint64 Remaining = End - Cursor;
-	if (Remaining == 0)
+
+	uint64 Remaining;
+	while (true)
 	{
-		return 0;
+		Remaining = End - Cursor;
+		if (Remaining)
+		{
+			break;
+		}
+
+		FPlatformProcess::Sleep(0.25f);
+		UpdateFileSize();
 	}
 
 	uint64 Size64 = Size;
