Fixed a nopch-nounity error.

#rnx

diff --git a/Engine/Source/Developer/TraceAnalysis/Private/DataStream.cpp b/Engine/Source/Developer/TraceAnalysis/Private/DataStream.cpp
--- a/Engine/Source/Developer/TraceAnalysis/Private/DataStream.cpp
+++ b/Engine/Source/Developer/TraceAnalysis/Private/DataStream.cpp
@@ -2,6 +2,7 @@
 
 #include "DataStream.h"
 #include "GenericPlatform/GenericPlatformFile.h"
+#include "HAL/PlatformProcess.h"
 
 namespace Trace
 {
