Make sure payload building can grow in size if need be.

#rnx

diff --git a/Engine/Source/Developer/TraceAnalysis/Private/Store/CborPayload.h b/Engine/Source/Developer/TraceAnalysis/Private/Store/CborPayload.h
--- a/Engine/Source/Developer/TraceAnalysis/Private/Store/CborPayload.h
+++ b/Engine/Source/Developer/TraceAnalysis/Private/Store/CborPayload.h
@@ -8,6 +8,8 @@
 
 #include "CborReader.h"
 #include "CborWriter.h"
+#include "Containers/Array.h"
+#include "Containers/ContainerAllocationPolicies.h"
 #include "Containers/StringView.h"
 #include "Serialization/MemoryArchive.h"
 #include "Serialization/MemoryReader.h"
@@ -20,10 +22,12 @@ template <int BufferSize>
 struct TInlineMemoryWriter
 	: public FMemoryArchive
 {
+
+	typedef TArray<uint8, TInlineAllocator<BufferSize>> BufferType;
+
 					TInlineMemoryWriter();
 	virtual void	Serialize(void* Data, int64 Num) override;
-	uint32			Used = 0;
-	uint8			Buffer[BufferSize];
+	BufferType		Buffer;
 };
 
 ////////////////////////////////////////////////////////////////////////////////
@@ -37,12 +41,7 @@ inline TInlineMemoryWriter<BufferSize>::TInlineMemoryWriter()
 template <int BufferSize>
 inline void TInlineMemoryWriter<BufferSize>::Serialize(void* Data, int64 Num)
 {
-	uint32 NextUsed = Used + uint32(Num);
-	if (NextUsed <= BufferSize)
-	{
-		memcpy(Buffer + Used, Data, Num);
-	}
-	Used = NextUsed;
+	Buffer.Append((const uint8*)Data, Num);;
 }
 
 
@@ -116,7 +115,7 @@ template <int Size>
 inline FPayload TPayloadBuilder<Size>::Done()
 {
 	CborWriter.WriteContainerEnd();
-	return { MemoryWriter.Buffer, MemoryWriter.Used };
+	return { MemoryWriter.Buffer.GetData(), uint32(MemoryWriter.Buffer.Num()) };
 }
 
 
