Don't truncate a trace's name quite to eagerly

#rnx

diff --git a/Engine/Source/Developer/TraceAnalysis/Private/Store/AsioStoreCborServer.cpp b/Engine/Source/Developer/TraceAnalysis/Private/Store/AsioStoreCborServer.cpp
--- a/Engine/Source/Developer/TraceAnalysis/Private/Store/AsioStoreCborServer.cpp
+++ b/Engine/Source/Developer/TraceAnalysis/Private/Store/AsioStoreCborServer.cpp
@@ -200,7 +200,7 @@ void FAsioStoreCborPeer::OnTraceInfo()
 	}
 
 	const TCHAR* Name = Trace->GetName();
-	char OutName[128];
+	char OutName[256];
 	for (char& Out : OutName)
 	{
 		Out = char(*Name++);
@@ -234,7 +234,18 @@ void FAsioStoreCborPeer::OnTraceRead()
 		return SendError(EStatusCode::BadRequest);
 	}
 
-	FAsioTraceRelay* Relay = new FAsioTraceRelay(GetIoContext(), Input);
+	uint32 SessionId = 0;
+	for (int i = 0, n = Recorder.GetSessionCount(); i < n; ++i)
+	{
+		const FAsioRecorder::FSession* Session = Recorder.GetSessionInfo(i);
+		if (Session->GetTraceId() == Id)
+		{
+			SessionId = Session->GetId();
+			break;
+		}
+	}
+
+	FAsioTraceRelay* Relay = new FAsioTraceRelay(GetIoContext(), Input, SessionId, Recorder);
 	uint32 Port = Relay->GetPort();
 	if (!Port)
 	{
