Bail earlier if there's no recipient sink for a completed IO operation.

#rnx

diff --git a/Engine/Source/Developer/TraceAnalysis/Private/Store/AsioIoable.cpp b/Engine/Source/Developer/TraceAnalysis/Private/Store/AsioIoable.cpp
--- a/Engine/Source/Developer/TraceAnalysis/Private/Store/AsioIoable.cpp
+++ b/Engine/Source/Developer/TraceAnalysis/Private/Store/AsioIoable.cpp
@@ -23,14 +23,14 @@ bool FAsioIoable::SetSink(FAsioIoSink* Ptr, uint32 Id)
 ////////////////////////////////////////////////////////////////////////////////
 void FAsioIoable::OnIoComplete(const asio::error_code& ErrorCode, int32 Size)
 {
-	if (ErrorCode)
+	if (SinkPtr == nullptr)
 	{
-		Size = -1;
+		return;
 	}
 
-	if (SinkPtr == nullptr)
+	if (ErrorCode)
 	{
-		return;
+		Size = -1;
 	}
 
 	FAsioIoSink* Ptr = SinkPtr;
