Encode error values in the Size parameter to aid debugging.

#rnx
#ushell-unshelve of 11071762

diff --git a/Engine/Source/Developer/TraceAnalysis/Private/Store/AsioIoable.cpp b/Engine/Source/Developer/TraceAnalysis/Private/Store/AsioIoable.cpp
--- a/Engine/Source/Developer/TraceAnalysis/Private/Store/AsioIoable.cpp
+++ b/Engine/Source/Developer/TraceAnalysis/Private/Store/AsioIoable.cpp
@@ -30,7 +30,7 @@ void FAsioIoable::OnIoComplete(const asio::error_code& ErrorCode, int32 Size)
 
 	if (ErrorCode)
 	{
-		Size = -1;
+		Size = 0 - ErrorCode.value();
 	}
 
 	FAsioIoSink* Ptr = SinkPtr;
