Wait for more data if a trace being relayed still being written to by a recorder.

#rnx
#ushell-unshelve of 11071762

diff --git a/Engine/Source/Developer/TraceAnalysis/Private/Store/AsioTraceRelay.cpp b/Engine/Source/Developer/TraceAnalysis/Private/Store/AsioTraceRelay.cpp
--- a/Engine/Source/Developer/TraceAnalysis/Private/Store/AsioTraceRelay.cpp
+++ b/Engine/Source/Developer/TraceAnalysis/Private/Store/AsioTraceRelay.cpp
@@ -4,15 +4,22 @@
 
 #if TRACE_WITH_ASIO
 
+#include "AsioRecorder.h"
 #include "AsioSocket.h"
 
 namespace Trace
 {
 
 ////////////////////////////////////////////////////////////////////////////////
-FAsioTraceRelay::FAsioTraceRelay(asio::io_context& IoContext, FAsioReadable* InInput)
+FAsioTraceRelay::FAsioTraceRelay(
+	asio::io_context& IoContext,
+	FAsioReadable* InInput,
+	uint32 InSessionId,
+	FAsioRecorder& InRecorder)
 : FAsioTcpServer(IoContext)
 , Input(InInput)
+, Recorder(InRecorder)
+, SessionId(InSessionId)
 {
 	StartServer();
 }
@@ -51,6 +58,19 @@ void FAsioTraceRelay::OnIoComplete(uint32 Id, int32 Size)
 {
 	if (Size < 0)
 	{
+		if (Id == OpRead && SessionId)
+		{
+			for (int i = 0, n = Recorder.GetSessionCount(); i < n; ++i)
+			{
+				const FAsioRecorder::FSession* Session = Recorder.GetSessionInfo(i);
+				if (Session->GetId() == SessionId)
+				{
+					FPlatformProcess::SleepNoStats(0.2f);
+					return OnIoComplete(OpStart, 0);
+				}
+			}
+		}
+
 		Output->Close();
 		Input->Close();
 		return;
diff --git a/Engine/Source/Developer/TraceAnalysis/Private/Store/AsioTraceRelay.h b/Engine/Source/Developer/TraceAnalysis/Private/Store/AsioTraceRelay.h
--- a/Engine/Source/Developer/TraceAnalysis/Private/Store/AsioTraceRelay.h
+++ b/Engine/Source/Developer/TraceAnalysis/Private/Store/AsioTraceRelay.h
@@ -12,13 +12,15 @@
 namespace Trace
 {
 
+class FAsioRecorder;
+
 ////////////////////////////////////////////////////////////////////////////////
 class FAsioTraceRelay
 	: public FAsioTcpServer
 	, public FAsioIoSink
 {
 public:
-						FAsioTraceRelay(asio::io_context& IoContext, FAsioReadable* InInput);
+						FAsioTraceRelay(asio::io_context& IoContext, FAsioReadable* InInput, uint32 InTraceIf, FAsioRecorder& InRecorder);
 						~FAsioTraceRelay();
 	void				Close();
 
@@ -29,6 +31,8 @@ private:
 	static const uint32	BufferSize = 64 << 10;
 	FAsioReadable*		Input;
 	FAsioWriteable*		Output = nullptr;
+	FAsioRecorder&		Recorder;
+	uint32				SessionId;
 	uint8				Buffer[BufferSize];
 };
 
