Always consider a trace live for now while store services are self-hosted.

#rnx
#ushell-unshelve of 11071762

diff --git a/Engine/Source/Developer/TraceAnalysis/Private/Store/AsioTraceRelay.cpp b/Engine/Source/Developer/TraceAnalysis/Private/Store/AsioTraceRelay.cpp
--- a/Engine/Source/Developer/TraceAnalysis/Private/Store/AsioTraceRelay.cpp
+++ b/Engine/Source/Developer/TraceAnalysis/Private/Store/AsioTraceRelay.cpp
@@ -58,6 +58,7 @@ void FAsioTraceRelay::OnIoComplete(uint32 Id, int32 Size)
 {
 	if (Size < 0)
 	{
+#if !defined(TRACE_WITH_ASIO)
 		if (Id == OpRead && SessionId)
 		{
 			for (int i = 0, n = Recorder.GetSessionCount(); i < n; ++i)
@@ -74,6 +75,14 @@ void FAsioTraceRelay::OnIoComplete(uint32 Id, int32 Size)
 		Output->Close();
 		Input->Close();
 		return;
+#else
+		if (Id == OpRead)
+		{
+			FPlatformProcess::SleepNoStats(0.2f);
+			OnIoComplete(OpStart, 0);
+		}
+		return;
+#endif
 	}
 
 	switch (Id)
