Added a block of disabled code that helps debug ASIO activity.

#rnx

diff --git a/Engine/Source/Developer/TraceAnalysis/Private/Store/AsioIoable.cpp b/Engine/Source/Developer/TraceAnalysis/Private/Store/AsioIoable.cpp
--- a/Engine/Source/Developer/TraceAnalysis/Private/Store/AsioIoable.cpp
+++ b/Engine/Source/Developer/TraceAnalysis/Private/Store/AsioIoable.cpp
@@ -28,6 +28,13 @@ void FAsioIoable::OnIoComplete(const asio::error_code& ErrorCode, int32 Size)
 		return;
 	}
 
+#if defined(UE_BUILD_DEBUG) && 0
+	std::string ErrorMessage;
+	{
+		ErrorMessage = ErrorCode.message();
+	}
+#endif
+
 	if (ErrorCode)
 	{
 		Size = 0 - ErrorCode.value();
