IsValid() -> IsOpen() to be more consistent with other trace-store classes.

#rnx

diff --git a/Engine/Source/Developer/TraceAnalysis/Private/Store/StoreClient.cpp b/Engine/Source/Developer/TraceAnalysis/Private/Store/StoreClient.cpp
--- a/Engine/Source/Developer/TraceAnalysis/Private/Store/StoreClient.cpp
+++ b/Engine/Source/Developer/TraceAnalysis/Private/Store/StoreClient.cpp
@@ -67,7 +67,7 @@ class FStoreCborClient
 public:
 							FStoreCborClient();
 							~FStoreCborClient();
-	bool					IsValid() const;
+	bool					IsOpen() const;
 	const FResponse&		GetResponse() const;
 	bool					Connect(const TCHAR* Host, uint16 Port);
 	bool					GetStatus();
@@ -99,7 +99,7 @@ FStoreCborClient::~FStoreCborClient()
 }
 
 ////////////////////////////////////////////////////////////////////////////////
-bool FStoreCborClient::IsValid() const
+bool FStoreCborClient::IsOpen() const
 {
 	return Socket.is_open();
 }
@@ -396,7 +396,7 @@ FStoreClient::~FStoreClient()
 bool FStoreClient::IsValid() const
 {
 #if TRACE_WITH_ASIO
-	return Impl->IsValid();
+	return Impl->IsOpen();
 #else
 	return false;
 #endif
