Added accessors to get the address and port of the store a client is connected to.

#rnx

diff --git a/Engine/Source/Developer/TraceAnalysis/Private/Store/StoreClient.cpp b/Engine/Source/Developer/TraceAnalysis/Private/Store/StoreClient.cpp
--- a/Engine/Source/Developer/TraceAnalysis/Private/Store/StoreClient.cpp
+++ b/Engine/Source/Developer/TraceAnalysis/Private/Store/StoreClient.cpp
@@ -68,6 +68,8 @@ public:
 							FStoreCborClient();
 							~FStoreCborClient();
 	bool					IsOpen() const;
+	uint32					GetStoreAddress() const;
+	uint32					GetStorePort() const;
 	const FResponse&		GetResponse() const;
 	bool					Connect(const TCHAR* Host, uint16 Port);
 	bool					GetStatus();
@@ -104,6 +106,31 @@ bool FStoreCborClient::IsOpen() const
 	return Socket.is_open();
 }
 
+////////////////////////////////////////////////////////////////////////////////
+uint32 FStoreCborClient::GetStoreAddress() const
+{
+	if (!IsOpen())
+	{
+		return 0;
+	}
+
+	const asio::ip::tcp::endpoint Endpoint = Socket.remote_endpoint();
+	asio::ip::address Address = Endpoint.address();
+	return Address.is_v4() ? Address.to_v4().to_uint() : 0;
+}
+
+////////////////////////////////////////////////////////////////////////////////
+uint32 FStoreCborClient::GetStorePort() const
+{
+	if (!IsOpen())
+	{
+		return 0;
+	}
+
+	const asio::ip::tcp::endpoint Endpoint = Socket.remote_endpoint();
+	return Endpoint.port();
+}
+
 ////////////////////////////////////////////////////////////////////////////////
 const FResponse& FStoreCborClient::GetResponse() const
 {
@@ -402,6 +429,26 @@ bool FStoreClient::IsValid() const
 #endif
 }
 
+////////////////////////////////////////////////////////////////////////////////
+uint32 FStoreClient::GetStoreAddress() const
+{
+#if TRACE_WITH_ASIO
+	return Impl->GetStoreAddress();
+#else
+	return 0;
+#endif
+}
+
+////////////////////////////////////////////////////////////////////////////////
+uint32 FStoreClient::GetStorePort() const
+{
+#if TRACE_WITH_ASIO
+	return Impl->GetStorePort();
+#else
+	return 0;
+#endif
+}
+
 ////////////////////////////////////////////////////////////////////////////////
 const FStoreClient::FStatus* FStoreClient::GetStatus()
 {
diff --git a/Engine/Source/Developer/TraceAnalysis/Public/Trace/StoreClient.h b/Engine/Source/Developer/TraceAnalysis/Public/Trace/StoreClient.h
--- a/Engine/Source/Developer/TraceAnalysis/Public/Trace/StoreClient.h
+++ b/Engine/Source/Developer/TraceAnalysis/Public/Trace/StoreClient.h
@@ -16,6 +16,9 @@ class IInDataStream;
 class TRACEANALYSIS_API FStoreClient
 {
 public:
+	uint32				GetStoreAddress() const;
+	uint32				GetStorePort() const;
+
 	struct TRACEANALYSIS_API FStatus
 	{
 		uint32			GetRecorderPort() const;
