Fixed reported trace timestamps

#rnx

diff --git a/Engine/Source/Developer/TraceAnalysis/Private/Store/AsioStore.cpp b/Engine/Source/Developer/TraceAnalysis/Private/Store/AsioStore.cpp
--- a/Engine/Source/Developer/TraceAnalysis/Private/Store/AsioStore.cpp
+++ b/Engine/Source/Developer/TraceAnalysis/Private/Store/AsioStore.cpp
@@ -42,7 +42,7 @@ uint64 FAsioStore::FTrace::GetTimestamp() const
 	{
 		return 0;
 	}
-	return (uint64(Time.dwHighDateTime) << 32ull) | uint64(Time.dwLowDateTime);
+	return (uint64(Time.dwLowDateTime) << 32ull) | uint64(Time.dwHighDateTime);
 }
 
 
