Added method to query if a connection to a store cbor-server is still open.

#rnx

diff --git a/Engine/Source/Developer/TraceAnalysis/Private/Store/AsioStoreCborServer.cpp b/Engine/Source/Developer/TraceAnalysis/Private/Store/AsioStoreCborServer.cpp
--- a/Engine/Source/Developer/TraceAnalysis/Private/Store/AsioStoreCborServer.cpp
+++ b/Engine/Source/Developer/TraceAnalysis/Private/Store/AsioStoreCborServer.cpp
@@ -23,7 +23,7 @@ class FAsioStoreCborPeer
 public:
 					FAsioStoreCborPeer(asio::ip::tcp::socket& InSocket, FAsioStore& InStore, FAsioRecorder& InRecorder);
 	virtual			~FAsioStoreCborPeer();
-	bool			IsActive() const;
+	bool			IsOpen() const;
 	void			Close();
 
 protected:
@@ -74,7 +74,7 @@ FAsioStoreCborPeer::~FAsioStoreCborPeer()
 }
 
 ////////////////////////////////////////////////////////////////////////////////
-bool FAsioStoreCborPeer::IsActive() const
+bool FAsioStoreCborPeer::IsOpen() const
 {
 	return Socket.IsOpen();
 }
@@ -82,7 +82,7 @@ bool FAsioStoreCborPeer::IsActive() const
 ////////////////////////////////////////////////////////////////////////////////
 void FAsioStoreCborPeer::Close()
 {
-	if (IsActive())
+	if (IsOpen())
 	{
 		Socket.Close();
 	}
