Added method to check it a trace relay is open or not.

#rnx

diff --git a/Engine/Source/Developer/TraceAnalysis/Private/Store/AsioTraceRelay.cpp b/Engine/Source/Developer/TraceAnalysis/Private/Store/AsioTraceRelay.cpp
--- a/Engine/Source/Developer/TraceAnalysis/Private/Store/AsioTraceRelay.cpp
+++ b/Engine/Source/Developer/TraceAnalysis/Private/Store/AsioTraceRelay.cpp
@@ -45,6 +45,12 @@ void FAsioTraceRelay::Close()
 	Input->Close();
 }
 
+////////////////////////////////////////////////////////////////////////////////
+bool FAsioTraceRelay::IsOpen() const
+{
+	return (Output != nullptr) && Output->IsOpen();
+}
+
 ////////////////////////////////////////////////////////////////////////////////
 bool FAsioTraceRelay::OnAccept(asio::ip::tcp::socket& Socket)
 {
diff --git a/Engine/Source/Developer/TraceAnalysis/Private/Store/AsioTraceRelay.h b/Engine/Source/Developer/TraceAnalysis/Private/Store/AsioTraceRelay.h
--- a/Engine/Source/Developer/TraceAnalysis/Private/Store/AsioTraceRelay.h
+++ b/Engine/Source/Developer/TraceAnalysis/Private/Store/AsioTraceRelay.h
@@ -13,6 +13,7 @@ namespace Trace
 {
 
 class FAsioRecorder;
+class FAsioSocket;
 
 ////////////////////////////////////////////////////////////////////////////////
 class FAsioTraceRelay
@@ -22,6 +23,7 @@ class FAsioTraceRelay
 public:
 						FAsioTraceRelay(asio::io_context& IoContext, FAsioReadable* InInput, uint32 InTraceIf, FAsioRecorder& InRecorder);
 						~FAsioTraceRelay();
+	bool				IsOpen() const;
 	void				Close();
 
 private:
@@ -30,7 +32,7 @@ private:
 	enum				{ OpStart, OpRead, OpSend };
 	static const uint32	BufferSize = 64 << 10;
 	FAsioReadable*		Input;
-	FAsioWriteable*		Output = nullptr;
+	FAsioSocket*		Output = nullptr;
 	FAsioRecorder&		Recorder;
 	uint32				SessionId;
 	uint8				Buffer[BufferSize];
