Added an IsOpen() for recorder relays.

#rnx

diff --git a/Engine/Source/Developer/TraceAnalysis/Private/Store/AsioRecorder.cpp b/Engine/Source/Developer/TraceAnalysis/Private/Store/AsioRecorder.cpp
--- a/Engine/Source/Developer/TraceAnalysis/Private/Store/AsioRecorder.cpp
+++ b/Engine/Source/Developer/TraceAnalysis/Private/Store/AsioRecorder.cpp
@@ -18,6 +18,7 @@ class FAsioRecorderRelay
 public:
 						FAsioRecorderRelay(asio::ip::tcp::socket& Socket, FAsioWriteable* InOutput);
 	virtual				~FAsioRecorderRelay();
+	bool				IsOpen();
 	void				Close();
 	uint32				GetIpAddress() const;
 
@@ -45,6 +46,12 @@ FAsioRecorderRelay::~FAsioRecorderRelay()
 	delete Output;
 }
 
+////////////////////////////////////////////////////////////////////////////////
+bool FAsioRecorderRelay::IsOpen()
+{
+	return Input.IsOpen();
+}
+
 ////////////////////////////////////////////////////////////////////////////////
 void FAsioRecorderRelay::Close()
 {
