The recorder's relays should close themselves on delete.

#rnx

diff --git a/Engine/Source/Developer/TraceAnalysis/Private/Store/AsioRecorder.cpp b/Engine/Source/Developer/TraceAnalysis/Private/Store/AsioRecorder.cpp
--- a/Engine/Source/Developer/TraceAnalysis/Private/Store/AsioRecorder.cpp
+++ b/Engine/Source/Developer/TraceAnalysis/Private/Store/AsioRecorder.cpp
@@ -43,6 +43,7 @@ FAsioRecorderRelay::FAsioRecorderRelay(asio::ip::tcp::socket& Socket, FAsioWrite
 ////////////////////////////////////////////////////////////////////////////////
 FAsioRecorderRelay::~FAsioRecorderRelay()
 {
+	Close();
 	delete Output;
 }
 
@@ -123,7 +124,6 @@ FAsioRecorder::~FAsioRecorder()
 {
 	for (FSession& Session : Sessions)
 	{
-		Session.Relay->Close();
 		delete Session.Relay;
 	}
 }
