Added a rolling clean up of a cbor-server's attached clients.

#rnx

diff --git a/Engine/Source/Developer/TraceAnalysis/Private/Store/AsioStoreCborServer.cpp b/Engine/Source/Developer/TraceAnalysis/Private/Store/AsioStoreCborServer.cpp
--- a/Engine/Source/Developer/TraceAnalysis/Private/Store/AsioStoreCborServer.cpp
+++ b/Engine/Source/Developer/TraceAnalysis/Private/Store/AsioStoreCborServer.cpp
@@ -349,10 +349,12 @@ FAsioStoreCborServer::FAsioStoreCborServer(
 	FAsioStore& InStore,
 	FAsioRecorder& InRecorder)
 : FAsioTcpServer(IoContext)
+, FAsioTickable(IoContext)
 , Store(InStore)
 , Recorder(InRecorder)
 {
 	StartServer();
+	StartTick(500);
 }
 
 ////////////////////////////////////////////////////////////////////////////////
@@ -385,6 +387,26 @@ bool FAsioStoreCborServer::OnAccept(asio::ip::tcp::socket& Socket)
 	return true;
 }
 
+////////////////////////////////////////////////////////////////////////////////
+void FAsioStoreCborServer::OnTick()
+{
+	uint32 FinalNum = 0;
+	for (int i = 0, n = Peers.Num(); i < n; ++i)
+	{
+		FAsioStoreCborPeer* Peer = Peers[i];
+		if (Peer->IsOpen())
+		{
+			++FinalNum;
+			continue;
+		}
+
+		delete Peer;
+		Peers[FinalNum] = Peer;
+	}
+
+	Peers.SetNum(FinalNum);
+}
+
 } // namespace Trace
 
 #endif // TRACE_WITH_ASIO
diff --git a/Engine/Source/Developer/TraceAnalysis/Private/Store/AsioStoreCborServer.h b/Engine/Source/Developer/TraceAnalysis/Private/Store/AsioStoreCborServer.h
--- a/Engine/Source/Developer/TraceAnalysis/Private/Store/AsioStoreCborServer.h
+++ b/Engine/Source/Developer/TraceAnalysis/Private/Store/AsioStoreCborServer.h
@@ -7,6 +7,7 @@
 #if TRACE_WITH_ASIO
 
 #include "AsioTcpServer.h"
+#include "AsioTickable.h"
 #include "Containers/Array.h"
 
 namespace Trace
@@ -19,6 +20,7 @@ class FAsioStoreCborPeer;
 ////////////////////////////////////////////////////////////////////////////////
 class FAsioStoreCborServer
 	: public FAsioTcpServer
+	, public FAsioTickable
 {
 public:
 								FAsioStoreCborServer(asio::io_context& IoContext, FAsioStore& InStore, FAsioRecorder& InRecorder);
@@ -28,6 +30,7 @@ public:
 
 private:
 	virtual bool				OnAccept(asio::ip::tcp::socket& Socket) override;
+	virtual void				OnTick() override;
 	TArray<FAsioStoreCborPeer*>	Peers;
 	FAsioStore&					Store;
 	FAsioRecorder&				Recorder;
