Implemented ticked clean up of a recorder sessions.

#rnx

diff --git a/Engine/Source/Developer/TraceAnalysis/Private/Store/AsioRecorder.cpp b/Engine/Source/Developer/TraceAnalysis/Private/Store/AsioRecorder.cpp
--- a/Engine/Source/Developer/TraceAnalysis/Private/Store/AsioRecorder.cpp
+++ b/Engine/Source/Developer/TraceAnalysis/Private/Store/AsioRecorder.cpp
@@ -175,7 +175,21 @@ bool FAsioRecorder::OnAccept(asio::ip::tcp::socket& Socket)
 ////////////////////////////////////////////////////////////////////////////////
 void FAsioRecorder::OnTick()
 {
-	/* cleanup dead clients */
+	uint32 FinalNum = 0;
+	for (int i = 0, n = Sessions.Num(); i < n; ++i)
+	{
+		FSession& Session = Sessions[i];
+		if (Session.Relay->IsOpen())
+		{
+			++FinalNum;
+			continue;
+		}
+
+		delete Session.Relay;
+		Sessions[FinalNum] = Session;
+	}
+
+	Sessions.SetNum(FinalNum);
 }
 
 } // namespace Trace
