Use the correct implementation for relaying traces where data's still being recorded.

#rnx

diff --git a/Engine/Source/Developer/TraceAnalysis/Private/Store/AsioTraceRelay.cpp b/Engine/Source/Developer/TraceAnalysis/Private/Store/AsioTraceRelay.cpp
--- a/Engine/Source/Developer/TraceAnalysis/Private/Store/AsioTraceRelay.cpp
+++ b/Engine/Source/Developer/TraceAnalysis/Private/Store/AsioTraceRelay.cpp
@@ -64,7 +64,6 @@ void FAsioTraceRelay::OnIoComplete(uint32 Id, int32 Size)
 {
 	if (Size < 0)
 	{
-#if !defined(TRACE_WITH_ASIO)
 		if (Id == OpRead && SessionId)
 		{
 			for (int i = 0, n = Recorder.GetSessionCount(); i < n; ++i)
@@ -81,14 +80,6 @@ void FAsioTraceRelay::OnIoComplete(uint32 Id, int32 Size)
 		Output->Close();
 		Input->Close();
 		return;
-#else
-		if (Id == OpRead)
-		{
-			FPlatformProcess::SleepNoStats(0.2f);
-			OnIoComplete(OpStart, 0);
-		}
-		return;
-#endif
 	}
 
 	switch (Id)
