Better cross-module memory management for FStoreClient.

#rnx

diff --git a/Engine/Source/Developer/TraceAnalysis/Private/Store/StoreClient.cpp b/Engine/Source/Developer/TraceAnalysis/Private/Store/StoreClient.cpp
--- a/Engine/Source/Developer/TraceAnalysis/Private/Store/StoreClient.cpp
+++ b/Engine/Source/Developer/TraceAnalysis/Private/Store/StoreClient.cpp
@@ -388,18 +388,10 @@ uint32 FStoreClient::FSessionInfo::GetIpAddress() const
 
 
 
-////////////////////////////////////////////////////////////////////////////////
-struct FStoreClient::FImpl
-#if TRACE_WITH_ASIO
-	: public FStoreCborClient
-#endif
-{
-};
-
 ////////////////////////////////////////////////////////////////////////////////
 FStoreClient* FStoreClient::Connect(const TCHAR* Host, uint32 Port)
 {
-	FImpl* Impl = new FStoreClient::FImpl();
+	FStoreCborClient* Impl = new FStoreCborClient();
 #if TRACE_WITH_ASIO
 	if (!Impl->Connect(Host, Port))
 #endif
@@ -408,22 +400,22 @@ FStoreClient* FStoreClient::Connect(const TCHAR* Host, uint32 Port)
 		return nullptr;
 	}
 
-	FStoreClient* Client = new FStoreClient();
-	Client->Impl = Impl;
-	return Client;
+	return (FStoreClient*)Impl;
 }
 
 ////////////////////////////////////////////////////////////////////////////////
-FStoreClient::~FStoreClient()
+void FStoreClient::operator delete (void* Addr)
 {
-	delete Impl;
+	auto* Self = (FStoreCborClient*)Addr;
+	delete Addr;
 }
 
 ////////////////////////////////////////////////////////////////////////////////
 bool FStoreClient::IsValid() const
 {
 #if TRACE_WITH_ASIO
-	return Impl->IsOpen();
+	auto* Self = (FStoreCborClient*)this;
+	return Self->IsOpen();
 #else
 	return false;
 #endif
@@ -433,7 +425,8 @@ bool FStoreClient::IsValid() const
 uint32 FStoreClient::GetStoreAddress() const
 {
 #if TRACE_WITH_ASIO
-	return Impl->GetStoreAddress();
+	auto* Self = (FStoreCborClient*)this;
+	return Self->GetStoreAddress();
 #else
 	return 0;
 #endif
@@ -443,7 +436,8 @@ uint32 FStoreClient::GetStoreAddress() const
 uint32 FStoreClient::GetStorePort() const
 {
 #if TRACE_WITH_ASIO
-	return Impl->GetStorePort();
+	auto* Self = (FStoreCborClient*)this;
+	return Self->GetStorePort();
 #else
 	return 0;
 #endif
@@ -453,12 +447,13 @@ uint32 FStoreClient::GetStorePort() const
 const FStoreClient::FStatus* FStoreClient::GetStatus()
 {
 #if TRACE_WITH_ASIO
-	if (!Impl->GetStatus())
+	auto* Self = (FStoreCborClient*)this;
+	if (!Self->GetStatus())
 	{
 		return nullptr;
 	}
 
-	const FResponse& Response = Impl->GetResponse();
+	const FResponse& Response = Self->GetResponse();
 	return (FStatus*)(&Response);
 #else
 	return nullptr;
@@ -469,12 +464,13 @@ const FStoreClient::FStatus* FStoreClient::GetStatus()
 uint32 FStoreClient::GetTraceCount()
 {
 #if TRACE_WITH_ASIO
-	if (!Impl->GetTraceCount())
+	auto* Self = (FStoreCborClient*)this;
+	if (!Self->GetTraceCount())
 	{
 		return 0;
 	}
 
-	return Impl->GetResponse().GetInteger("count", 0);
+	return Self->GetResponse().GetInteger("count", 0);
 #else
 	return 0;
 #endif // TRACE_WITH_ASIO
@@ -484,12 +480,13 @@ uint32 FStoreClient::GetTraceCount()
 const FStoreClient::FTraceInfo* FStoreClient::GetTraceInfo(uint32 Index)
 {
 #if TRACE_WITH_ASIO
-	if (!Impl->GetTraceInfo(Index))
+	auto* Self = (FStoreCborClient*)this;
+	if (!Self->GetTraceInfo(Index))
 	{
 		return nullptr;
 	}
 
-	const FResponse& Response = Impl->GetResponse();
+	const FResponse& Response = Self->GetResponse();
 	return (FTraceInfo*)(&Response);
 #else
 	return nullptr;
@@ -500,12 +497,13 @@ const FStoreClient::FTraceInfo* FStoreClient::GetTraceInfo(uint32 Index)
 const FStoreClient::FTraceInfo* FStoreClient::GetTraceInfoById(uint32 Id)
 {
 #if TRACE_WITH_ASIO
-	if (!Impl->GetTraceInfoById(Id))
+	auto* Self = (FStoreCborClient*)this;
+	if (!Self->GetTraceInfoById(Id))
 	{
 		return nullptr;
 	}
 
-	const FResponse& Response = Impl->GetResponse();
+	const FResponse& Response = Self->GetResponse();
 	return (FTraceInfo*)(&Response);
 #else
 	return nullptr;
@@ -516,7 +514,8 @@ const FStoreClient::FTraceInfo* FStoreClient::GetTraceInfoById(uint32 Id)
 FStoreClient::FTraceData FStoreClient::ReadTrace(uint32 Id)
 {
 #if TRACE_WITH_ASIO
-	return FTraceData(Impl->ReadTrace(Id));
+	auto* Self = (FStoreCborClient*)this;
+	return FTraceData(Self->ReadTrace(Id));
 #else
 	return nullptr;
 #endif // TRACE_WITH_ASIO
@@ -526,12 +525,13 @@ FStoreClient::FTraceData FStoreClient::ReadTrace(uint32 Id)
 uint32 FStoreClient::GetSessionCount() const
 {
 #if TRACE_WITH_ASIO
-	if (!Impl->GetSessionCount())
+	auto* Self = (FStoreCborClient*)this;
+	if (!Self->GetSessionCount())
 	{
 		return 0;
 	}
 
-	return Impl->GetResponse().GetInteger("count", 0);
+	return Self->GetResponse().GetInteger("count", 0);
 #else
 	return 0;
 #endif // TRACE_WITH_ASIO
@@ -541,12 +541,13 @@ uint32 FStoreClient::GetSessionCount() const
 const FStoreClient::FSessionInfo* FStoreClient::GetSessionInfo(uint32 Index) const
 {
 #if TRACE_WITH_ASIO
-	if (!Impl->GetSessionInfo(Index))
+	auto* Self = (FStoreCborClient*)this;
+	if (!Self->GetSessionInfo(Index))
 	{
 		return nullptr;
 	}
 
-	const FResponse& Response = Impl->GetResponse();
+	const FResponse& Response = Self->GetResponse();
 	return (FSessionInfo*)(&Response);
 #else
 	return nullptr;
@@ -557,12 +558,13 @@ const FStoreClient::FSessionInfo* FStoreClient::GetSessionInfo(uint32 Index) con
 const FStoreClient::FSessionInfo* FStoreClient::GetSessionInfoById(uint32 Id) const
 {
 #if TRACE_WITH_ASIO
-	if (!Impl->GetSessionInfoById(Id))
+	auto* Self = (FStoreCborClient*)this;
+	if (!Self->GetSessionInfoById(Id))
 	{
 		return nullptr;
 	}
 
-	const FResponse& Response = Impl->GetResponse();
+	const FResponse& Response = Self->GetResponse();
 	return (FSessionInfo*)(&Response);
 #else
 	return nullptr;
@@ -573,12 +575,13 @@ const FStoreClient::FSessionInfo* FStoreClient::GetSessionInfoById(uint32 Id) co
 const FStoreClient::FSessionInfo* FStoreClient::GetSessionInfoByTraceId(uint32 TraceId) const
 {
 #if TRACE_WITH_ASIO
-	if (!Impl->GetSessionInfoByTraceId(TraceId))
+	auto* Self = (FStoreCborClient*)this;
+	if (!Self->GetSessionInfoByTraceId(TraceId))
 	{
 		return nullptr;
 	}
 
-	const FResponse& Response = Impl->GetResponse();
+	const FResponse& Response = Self->GetResponse();
 	return (FSessionInfo*)(&Response);
 #else
 	return nullptr;
diff --git a/Engine/Source/Developer/TraceAnalysis/Public/Trace/StoreClient.h b/Engine/Source/Developer/TraceAnalysis/Public/Trace/StoreClient.h
--- a/Engine/Source/Developer/TraceAnalysis/Public/Trace/StoreClient.h
+++ b/Engine/Source/Developer/TraceAnalysis/Public/Trace/StoreClient.h
@@ -40,8 +40,9 @@ public:
 		using TUniquePtr<IInDataStream>::TUniquePtr;
 	};
 	
-						~FStoreClient();
+						~FStoreClient() = default;
 	static FStoreClient*Connect(const TCHAR* Host, uint32 Port);
+	void				operator delete (void* Addr);
 	bool				IsValid() const;
 	const FStatus*		GetStatus();
 	uint32				GetTraceCount();
@@ -78,10 +79,6 @@ public:
 
 private:
 						FStoreClient() = default;
-	struct				FImpl;
-	FImpl*				Impl = nullptr;
-
-private:
 						FStoreClient(const FStoreClient&) = delete;
 						FStoreClient(const FStoreClient&&) = delete;
 	void				operator = (const FStoreClient&) = delete;
