Stopping an ASIO context is now an explicit choice instead of doing it in the destructor to better control tearing down a context and anything using it.

#rnx

diff --git a/Engine/Source/Developer/TraceAnalysis/Private/Store/AsioContext.cpp b/Engine/Source/Developer/TraceAnalysis/Private/Store/AsioContext.cpp
--- a/Engine/Source/Developer/TraceAnalysis/Private/Store/AsioContext.cpp
+++ b/Engine/Source/Developer/TraceAnalysis/Private/Store/AsioContext.cpp
@@ -20,19 +20,38 @@ FAsioContext::FAsioContext(int32 ThreadCount)
 	ThreadPool.SetNum(ThreadCount);
 }
 
+////////////////////////////////////////////////////////////////////////////////
+FAsioContext::~FAsioContext()
+{
+	Stop();
+	delete IoContext;
+}
+
 ////////////////////////////////////////////////////////////////////////////////
 void FAsioContext::Start()
 {
+	if (bRunning)
+	{
+		return;
+	}
+
 	for (std::thread& Thread : ThreadPool)
 	{
 		std::thread TempThread([this] () { IoContext->run(); });
 		Thread = MoveTemp(TempThread);
 	}
+
+	bRunning = true;
 }
 
 ////////////////////////////////////////////////////////////////////////////////
-FAsioContext::~FAsioContext()
+void FAsioContext::Stop()
 {
+	if (!bRunning)
+	{
+		return;
+	}
+
 	IoContext->stop();
 
 	for (std::thread& Thread : ThreadPool)
@@ -40,7 +59,7 @@ FAsioContext::~FAsioContext()
 		Thread.join();
 	}
 
-	delete IoContext;
+	bRunning = false;
 }
 
 ////////////////////////////////////////////////////////////////////////////////
diff --git a/Engine/Source/Developer/TraceAnalysis/Private/Store/AsioContext.h b/Engine/Source/Developer/TraceAnalysis/Private/Store/AsioContext.h
--- a/Engine/Source/Developer/TraceAnalysis/Private/Store/AsioContext.h
+++ b/Engine/Source/Developer/TraceAnalysis/Private/Store/AsioContext.h
@@ -20,10 +20,12 @@ public:
 							~FAsioContext();
 	asio::io_context&		Get();
 	void					Start();
+	void					Stop();
 
 private:
 	asio::io_context*		IoContext;
 	TArray<std::thread>		ThreadPool;
+	bool					bRunning = false;
 };
 
 } // namespace Trace
