Better cross-module memory management for FStoreService.

#rnx

diff --git a/Engine/Source/Developer/TraceAnalysis/Private/Store/StoreService.cpp b/Engine/Source/Developer/TraceAnalysis/Private/Store/StoreService.cpp
--- a/Engine/Source/Developer/TraceAnalysis/Private/Store/StoreService.cpp
+++ b/Engine/Source/Developer/TraceAnalysis/Private/Store/StoreService.cpp
@@ -15,10 +15,10 @@ namespace Trace
 {
 
 ////////////////////////////////////////////////////////////////////////////////
-struct FStoreService::FImpl
+struct FStoreServiceImpl
 {
 public:
-							FImpl(const FDesc& Desc);
+							FStoreServiceImpl(const FStoreService::FDesc& Desc);
 	FAsioContext			Context;
 	FAsioStore				Store;
 	FAsioRecorder			Recorder;
@@ -26,7 +26,7 @@ public:
 };
 
 ////////////////////////////////////////////////////////////////////////////////
-FStoreService::FImpl::FImpl(const FDesc& Desc)
+FStoreServiceImpl::FStoreServiceImpl(const FStoreService::FDesc& Desc)
 : Context(Desc.ThreadCount)
 , Store(Context.Get(), Desc.StoreDir)
 , Recorder(Context.Get(), Store)
@@ -43,14 +43,6 @@ FStoreService::FImpl::FImpl(const FDesc& Desc)
 namespace Trace
 {
 
-////////////////////////////////////////////////////////////////////////////////
-FStoreService::~FStoreService()
-{
-#if TRACE_WITH_ASIO
-	delete Impl;
-#endif
-}
-
 ////////////////////////////////////////////////////////////////////////////////
 FStoreService* FStoreService::Create(const FDesc& InDesc)
 {
@@ -68,7 +60,7 @@ FStoreService* FStoreService::Create(const FDesc& InDesc)
 	// TODO: not thread safe yet
 	Desc.ThreadCount = 1;
 
-	FStoreService::FImpl* Impl = new FStoreService::FImpl(Desc);
+	FStoreServiceImpl* Impl = new FStoreServiceImpl(Desc);
 	if (Desc.RecorderPort >= 0)
 	{
 		FAsioRecorder& Recorder = Impl->Recorder;
@@ -77,19 +69,25 @@ FStoreService* FStoreService::Create(const FDesc& InDesc)
 
 	Impl->Context.Start();
 
-	FStoreService* Server = new FStoreService();
-	Server->Impl = Impl;
-	return Server;
+	return (FStoreService*)Impl;
 #else
 	return nullptr;
 #endif // TRACE_WITH_ASIO
 }
 
+////////////////////////////////////////////////////////////////////////////////
+void FStoreService::operator delete (void* Addr)
+{
+	auto* Self = (FStoreServiceImpl*)Addr;
+	delete Self;
+}
+
 ////////////////////////////////////////////////////////////////////////////////
 uint32 FStoreService::GetPort() const
 {
 #if TRACE_WITH_ASIO
-	return Impl->CborServer.GetPort();
+	auto* Self = (FStoreServiceImpl*)this;
+	return Self->CborServer.GetPort();
 #else
 	return 0;
 #endif // TRACE_WITH_ASIO
diff --git a/Engine/Source/Developer/TraceAnalysis/Public/Trace/StoreService.h b/Engine/Source/Developer/TraceAnalysis/Public/Trace/StoreService.h
--- a/Engine/Source/Developer/TraceAnalysis/Public/Trace/StoreService.h
+++ b/Engine/Source/Developer/TraceAnalysis/Public/Trace/StoreService.h
@@ -18,16 +18,13 @@ public:
 		int32				ThreadCount  = 0; // <=0:logical CPU count
 	};
 
-							~FStoreService();
+							~FStoreService() = default;
 	static FStoreService*	Create(const FDesc& Desc);
+	void					operator delete (void* Addr);
 	uint32					GetPort() const;
 
 private:
 							FStoreService() = default;
-	struct					FImpl;
-	FImpl*					Impl;
-
-private:
 							FStoreService(const FStoreService&) = delete;
 							FStoreService(const FStoreService&&) = delete;
 	void					operator = (const FStoreService&) = delete;
