Fixed the store-service segfaulting on shutdown.

#rnx

diff --git a/Engine/Source/Developer/TraceAnalysis/Private/Store/StoreService.cpp b/Engine/Source/Developer/TraceAnalysis/Private/Store/StoreService.cpp
--- a/Engine/Source/Developer/TraceAnalysis/Private/Store/StoreService.cpp
+++ b/Engine/Source/Developer/TraceAnalysis/Private/Store/StoreService.cpp
@@ -19,6 +19,7 @@ struct FStoreServiceImpl
 {
 public:
 							FStoreServiceImpl(const FStoreService::FDesc& Desc);
+							~FStoreServiceImpl();
 	FAsioContext			Context;
 	FAsioStore				Store;
 	FAsioRecorder			Recorder;
@@ -34,6 +35,12 @@ FStoreServiceImpl::FStoreServiceImpl(const FStoreService::FDesc& Desc)
 {
 }
 
+////////////////////////////////////////////////////////////////////////////////
+FStoreServiceImpl::~FStoreServiceImpl()
+{
+	Context.Stop();
+}
+
 } // namespace Trace
 
 #endif // TRACE_WITH_ASIO
