Fixed compile error.

#rnx

diff --git a/Engine/Source/Developer/TraceAnalysis/Private/Store/StoreClient.cpp b/Engine/Source/Developer/TraceAnalysis/Private/Store/StoreClient.cpp
--- a/Engine/Source/Developer/TraceAnalysis/Private/Store/StoreClient.cpp
+++ b/Engine/Source/Developer/TraceAnalysis/Private/Store/StoreClient.cpp
@@ -407,7 +407,7 @@ FStoreClient* FStoreClient::Connect(const TCHAR* Host, uint32 Port)
 void FStoreClient::operator delete (void* Addr)
 {
 	auto* Self = (FStoreCborClient*)Addr;
-	delete Addr;
+	delete Self;
 }
 
 ////////////////////////////////////////////////////////////////////////////////
