[TraceAnalysis] Added workarounds to compile asio on Linux.

diff --git a/Engine/Source/Developer/TraceAnalysis/Private/Store/AsioFile.cpp b/Engine/Source/Developer/TraceAnalysis/Private/Store/AsioFile.cpp
--- a/Engine/Source/Developer/TraceAnalysis/Private/Store/AsioFile.cpp
+++ b/Engine/Source/Developer/TraceAnalysis/Private/Store/AsioFile.cpp
@@ -4,19 +4,32 @@
 
 #if TRACE_WITH_ASIO
 
+#if !PLATFORM_WINDOWS
+#include <sys/file.h>
+#include <errno.h>
+#endif
+
 namespace Trace
 {
 
 ////////////////////////////////////////////////////////////////////////////////
 FAsioFile::FAsioFile(asio::io_context& IoContext, UPTRINT OsHandle)
+#if PLATFORM_WINDOWS
 : Handle(IoContext, HANDLE(OsHandle))
+#else
+: StreamDescriptor(IoContext, OsHandle)
+#endif
 {
 }
 
 ////////////////////////////////////////////////////////////////////////////////
 void FAsioFile::Close()
 {
+#if PLATFORM_WINDOWS
 	Handle.close();
+#else
+	StreamDescriptor.close();
+#endif
 }
 
 ////////////////////////////////////////////////////////////////////////////////
@@ -27,9 +40,14 @@ bool FAsioFile::Write(const void* Src, uint32 Size, FAsioIoSink* Sink, uint32 Id
 		return false;
 	}
 
+#if PLATFORM_WINDOWS
 	asio::async_write_at(
 		Handle,
 		Offset,
+#else
+	asio::async_write(
+		StreamDescriptor,
+#endif
 		asio::buffer(Src, Size),
 		[this] (const asio::error_code& ErrorCode, size_t BytesWritten)
 		{
@@ -55,8 +73,12 @@ bool FAsioFile::ReadSome(void* Dest, uint32 DestSize, FAsioIoSink* Sink, uint32
 		return false;
 	}
 
+#if PLATFORM_WINDOWS
 	Handle.async_read_some_at(
 		Offset,
+#else
+	StreamDescriptor.async_read_some(
+#endif
 		asio::buffer(Dest, DestSize),
 		[this] (const asio::error_code& ErrorCode, size_t BytesRead)
 		{
@@ -71,27 +93,43 @@ bool FAsioFile::ReadSome(void* Dest, uint32 DestSize, FAsioIoSink* Sink, uint32
 ////////////////////////////////////////////////////////////////////////////////
 FAsioWriteable* FAsioFile::WriteFile(asio::io_context& IoContext, const TCHAR* Path)
 {
+#if PLATFORM_WINDOWS
 	HANDLE Handle = CreateFileW(Path, GENERIC_WRITE, FILE_SHARE_READ, nullptr,
 		CREATE_ALWAYS, FILE_FLAG_OVERLAPPED|FILE_ATTRIBUTE_NORMAL, nullptr);
 	if (Handle == INVALID_HANDLE_VALUE)
 	{
 		return nullptr;
 	}
-
 	return new FAsioFile(IoContext, UPTRINT(Handle));
+#else
+	int File = open(TCHAR_TO_ANSI(Path), O_WRONLY | O_CREAT, 0666);
+	if (!File)
+	{
+		return nullptr;
+	}
+	return new FAsioFile(IoContext, UPTRINT(File));
+#endif
 }
 
 ////////////////////////////////////////////////////////////////////////////////
 FAsioReadable* FAsioFile::ReadFile(asio::io_context& IoContext, const TCHAR* Path)
 {
+#if PLATFORM_WINDOWS
 	HANDLE Handle = CreateFileW(Path, GENERIC_READ, FILE_SHARE_READ|FILE_SHARE_WRITE,
 		nullptr, OPEN_EXISTING, FILE_FLAG_OVERLAPPED, nullptr);
 	if (Handle == INVALID_HANDLE_VALUE)
 	{
 		return nullptr;
 	}
-
 	return new FAsioFile(IoContext, UPTRINT(Handle));
+#else
+	int File = open(TCHAR_TO_ANSI(Path), O_RDONLY, 0666);
+	if (!File)
+	{
+		return nullptr;
+	}
+	return new FAsioFile(IoContext, UPTRINT(File));
+#endif
 }
 
 } // namespace Trace
diff --git a/Engine/Source/Developer/TraceAnalysis/Private/Store/AsioFile.h b/Engine/Source/Developer/TraceAnalysis/Private/Store/AsioFile.h
--- a/Engine/Source/Developer/TraceAnalysis/Private/Store/AsioFile.h
+++ b/Engine/Source/Developer/TraceAnalysis/Private/Store/AsioFile.h
@@ -7,6 +7,9 @@
 #if TRACE_WITH_ASIO
 
 #include "AsioIoable.h"
+#if !PLATFORM_WINDOWS
+#include "asio/posix/stream_descriptor.hpp"
+#endif
 
 namespace Trace
 {
@@ -26,7 +29,11 @@ public:
 	virtual bool						ReadSome(void* Dest, uint32 BufferSize, FAsioIoSink* Sink, uint32 Id) override;
 
 private:
+#if PLATFORM_WINDOWS
 	asio::windows::random_access_handle	Handle;
+#else
+	asio::posix::stream_descriptor		StreamDescriptor;
+#endif
 	uint64								Offset = 0;
 };
 
diff --git a/Engine/Source/Developer/TraceAnalysis/Private/Store/AsioStore.cpp b/Engine/Source/Developer/TraceAnalysis/Private/Store/AsioStore.cpp
--- a/Engine/Source/Developer/TraceAnalysis/Private/Store/AsioStore.cpp
+++ b/Engine/Source/Developer/TraceAnalysis/Private/Store/AsioStore.cpp
@@ -27,15 +27,20 @@ uint32 FAsioStore::FTrace::GetId() const
 ////////////////////////////////////////////////////////////////////////////////
 uint64 FAsioStore::FTrace::GetSize() const
 {
+#if PLATFORM_WINDOWS
 	HANDLE Inner = HANDLE(Handle);
 	LARGE_INTEGER FileSize;
 	GetFileSizeEx(Inner, &FileSize);
 	return FileSize.QuadPart;
+#else
+	return 0; // TODO
+#endif
 }
 
 ////////////////////////////////////////////////////////////////////////////////
 uint64 FAsioStore::FTrace::GetTimestamp() const
 {
+#if PLATFORM_WINDOWS
 	HANDLE Inner = HANDLE(Handle);
 	FILETIME Time;
 	if (!GetFileTime(Inner, &Time, nullptr, nullptr))
@@ -43,6 +48,9 @@ uint64 FAsioStore::FTrace::GetTimestamp() const
 		return 0;
 	}
 	return (uint64(Time.dwLowDateTime) << 32ull) | uint64(Time.dwHighDateTime);
+#else
+	return 0; // TODO
+#endif
 }
 
 
@@ -118,12 +126,20 @@ FAsioStore::FTrace* FAsioStore::GetTrace(uint32 Id)
 ////////////////////////////////////////////////////////////////////////////////
 FAsioStore::FTrace* FAsioStore::AddTrace(const TCHAR* Path)
 {
+#if PLATFORM_WINDOWS
 	HANDLE Handle = CreateFileW(Path, 0, FILE_SHARE_READ|FILE_SHARE_WRITE|FILE_SHARE_DELETE,
 		nullptr, OPEN_EXISTING, 0, nullptr);
 	if (Handle == INVALID_HANDLE_VALUE)
 	{
 		return nullptr;
 	}
+#else
+	int Handle = open(TCHAR_TO_ANSI(Path), O_RDWR, 0666);
+	if (!Handle)
+	{
+		return nullptr;
+	}
+#endif
 
 	const TCHAR* Dot = FCString::Strrchr(Path, '.');
 	if (Dot == nullptr)
