A few compile fixes that occur when ASIO is disabled.

#rnx

diff --git a/Engine/Source/Developer/TraceAnalysis/Private/Store/StoreClient.cpp b/Engine/Source/Developer/TraceAnalysis/Private/Store/StoreClient.cpp
--- a/Engine/Source/Developer/TraceAnalysis/Private/Store/StoreClient.cpp
+++ b/Engine/Source/Developer/TraceAnalysis/Private/Store/StoreClient.cpp
@@ -391,23 +391,27 @@ uint32 FStoreClient::FSessionInfo::GetIpAddress() const
 ////////////////////////////////////////////////////////////////////////////////
 FStoreClient* FStoreClient::Connect(const TCHAR* Host, uint32 Port)
 {
-	FStoreCborClient* Impl = new FStoreCborClient();
 #if TRACE_WITH_ASIO
+	FStoreCborClient* Impl = new FStoreCborClient();
 	if (!Impl->Connect(Host, Port))
-#endif
 	{
 		delete Impl;
 		return nullptr;
 	}
 
 	return (FStoreClient*)Impl;
+#else
+	return nullptr;
+#endif
 }
 
 ////////////////////////////////////////////////////////////////////////////////
 void FStoreClient::operator delete (void* Addr)
 {
+#if TRACE_WITH_ASIO
 	auto* Self = (FStoreCborClient*)Addr;
 	delete Self;
+#endif
 }
 
 ////////////////////////////////////////////////////////////////////////////////
diff --git a/Engine/Source/Developer/TraceAnalysis/Private/Store/StoreService.cpp b/Engine/Source/Developer/TraceAnalysis/Private/Store/StoreService.cpp
--- a/Engine/Source/Developer/TraceAnalysis/Private/Store/StoreService.cpp
+++ b/Engine/Source/Developer/TraceAnalysis/Private/Store/StoreService.cpp
@@ -85,8 +85,10 @@ FStoreService* FStoreService::Create(const FDesc& InDesc)
 ////////////////////////////////////////////////////////////////////////////////
 void FStoreService::operator delete (void* Addr)
 {
+#if TRACE_WITH_ASIO
 	auto* Self = (FStoreServiceImpl*)Addr;
 	delete Self;
+#endif // TRACE_WITH_ASIO
 }
 
 ////////////////////////////////////////////////////////////////////////////////
