[TraceAnalysis]
  - Fixed TPayloadBuilder to always write nullterminator for strings. It fixes trace file names having last char truncated.
  - Fixed FResponse::GetInteger to not truncate int64 to int32.

diff --git a/Engine/Source/Developer/TraceAnalysis/Private/Store/CborPayload.h b/Engine/Source/Developer/TraceAnalysis/Private/Store/CborPayload.h
--- a/Engine/Source/Developer/TraceAnalysis/Private/Store/CborPayload.h
+++ b/Engine/Source/Developer/TraceAnalysis/Private/Store/CborPayload.h
@@ -85,7 +85,7 @@ template <int N>
 inline TPayloadBuilder<Size>::TPayloadBuilder(const char (&MethodName)[N])
 {
 	CborWriter.WriteContainerStart(ECborCode::Map, -1);
-	AddString("$method", MethodName, N);
+	AddString("$method", MethodName, N - 1);
 }
 
 ////////////////////////////////////////////////////////////////////////////////
@@ -107,7 +107,7 @@ inline void TPayloadBuilder<Size>::AddString(
 {
 	Length = (Length < 0) ? int32(strlen(Value)) : Length;
 	CborWriter.WriteValue(Name, N);
-	CborWriter.WriteValue(Value, Length);
+	CborWriter.WriteValue(Value, Length + 1);
 }
 
 ////////////////////////////////////////////////////////////////////////////////
@@ -201,7 +201,7 @@ inline int64 FResponse::GetInteger(const char* Key, int64 Default) const
 		Default,
 		[this] (const FCborContext& Context, int64)
 		{
-			return int32(Context.AsInt());
+			return Context.AsInt();
 		}
 	);
 }
