[TraceAnalysis] Reverted change re timestamps as problem was in other code (serialization).

diff --git a/Engine/Source/Developer/TraceAnalysis/Private/Store/AsioStore.cpp b/Engine/Source/Developer/TraceAnalysis/Private/Store/AsioStore.cpp
--- a/Engine/Source/Developer/TraceAnalysis/Private/Store/AsioStore.cpp
+++ b/Engine/Source/Developer/TraceAnalysis/Private/Store/AsioStore.cpp
@@ -47,7 +47,7 @@ uint64 FAsioStore::FTrace::GetTimestamp() const
 	{
 		return 0;
 	}
-	return (uint64(Time.dwLowDateTime) << 32ull) | uint64(Time.dwHighDateTime);
+	return (uint64(Time.dwHighDateTime) << 32ull) | uint64(Time.dwLowDateTime);
 #else
 	return 0; // TODO
 #endif
