[TraceAnalysis] Adjusted trace timestamps to use FDateTime ticks (relative to year 1 UTC) instead of Windows FILETIME ticks (relative to year 1601 UTC).

diff --git a/Engine/Source/Developer/TraceAnalysis/Private/Store/AsioStore.cpp b/Engine/Source/Developer/TraceAnalysis/Private/Store/AsioStore.cpp
--- a/Engine/Source/Developer/TraceAnalysis/Private/Store/AsioStore.cpp
+++ b/Engine/Source/Developer/TraceAnalysis/Private/Store/AsioStore.cpp
@@ -47,7 +47,11 @@ uint64 FAsioStore::FTrace::GetTimestamp() const
 	{
 		return 0;
 	}
-	return (uint64(Time.dwHighDateTime) << 32ull) | uint64(Time.dwLowDateTime);
+	// Windows FILETIME is a 64-bit value that represents the number of 100-nanosecond intervals that have elapsed since 12:00 A.M.January 1, 1601 Coordinated Universal Time(UTC).
+	// We adjust it to be compatible with the FDateTime ticks number of 100-nanosecond intervals that have elapsed since 12:00 A.M.January 1, 0001 Coordinated Universal Time(UTC).
+	const uint64 WinTicks = (static_cast<uint64>(Time.dwHighDateTime) << 32ull) | static_cast<uint64>(Time.dwLowDateTime);
+	const uint64 Year1601 = 504911232000000000ULL; // FDateTime(1601, 1, 1).GetTicks()
+	return Year1601 + WinTicks;
 #else
 	return 0; // TODO
 #endif
