[TraceAnalysis] Removed TRACE_WITH_ASIO toggle macro.

diff --git a/Engine/Source/Developer/TraceAnalysis/Private/Asio/Asio.cpp b/Engine/Source/Developer/TraceAnalysis/Private/Asio/Asio.cpp
--- a/Engine/Source/Developer/TraceAnalysis/Private/Asio/Asio.cpp
+++ b/Engine/Source/Developer/TraceAnalysis/Private/Asio/Asio.cpp
@@ -2,8 +2,6 @@
 
 #include "Asio.h"
 
-#if TRACE_WITH_ASIO
-
 #include "HAL/Platform.h"
 #include "Misc/AssertionMacros.h"
 
@@ -31,5 +29,3 @@ THIRD_PARTY_INCLUDES_END
 #	include "Windows/HideWindowsPlatformAtomics.h"
 #	include "Windows/HideWindowsPlatformTypes.h"
 #endif
-
-#endif // TRACE_WITH_ASIO
diff --git a/Engine/Source/Developer/TraceAnalysis/Private/Asio/Asio.h b/Engine/Source/Developer/TraceAnalysis/Private/Asio/Asio.h
--- a/Engine/Source/Developer/TraceAnalysis/Private/Asio/Asio.h
+++ b/Engine/Source/Developer/TraceAnalysis/Private/Asio/Asio.h
@@ -2,12 +2,6 @@
 
 #pragma once
 
-#if !defined(TRACE_WITH_ASIO)
-#	define TRACE_WITH_ASIO 1
-#endif // 0
-
-#if TRACE_WITH_ASIO
-
 #include "HAL/Platform.h"
 
 #if PLATFORM_WINDOWS
@@ -46,5 +40,3 @@ THIRD_PARTY_INCLUDES_END
 #	include "Windows/HideWindowsPlatformAtomics.h"
 #	include "Windows/HideWindowsPlatformTypes.h"
 #endif
-
-#endif // TRACE_WITH_ASIO
diff --git a/Engine/Source/Developer/TraceAnalysis/Private/Store/AsioContext.cpp b/Engine/Source/Developer/TraceAnalysis/Private/Store/AsioContext.cpp
--- a/Engine/Source/Developer/TraceAnalysis/Private/Store/AsioContext.cpp
+++ b/Engine/Source/Developer/TraceAnalysis/Private/Store/AsioContext.cpp
@@ -2,8 +2,6 @@
 
 #include "AsioContext.h"
 
-#if TRACE_WITH_ASIO
-
 namespace Trace
 {
 
@@ -69,5 +67,3 @@ asio::io_context& FAsioContext::Get()
 }
 
 } // namespace Trace
-
-#endif // TRACE_WITH_ASIO
diff --git a/Engine/Source/Developer/TraceAnalysis/Private/Store/AsioContext.h b/Engine/Source/Developer/TraceAnalysis/Private/Store/AsioContext.h
--- a/Engine/Source/Developer/TraceAnalysis/Private/Store/AsioContext.h
+++ b/Engine/Source/Developer/TraceAnalysis/Private/Store/AsioContext.h
@@ -5,8 +5,6 @@
 #include "Asio/Asio.h"
 #include "Containers/Array.h"
 
-#if TRACE_WITH_ASIO
-
 #include <thread>
 
 namespace Trace
@@ -29,5 +27,3 @@ private:
 };
 
 } // namespace Trace
-
-#endif // TRACE_WITH_ASIO
diff --git a/Engine/Source/Developer/TraceAnalysis/Private/Store/AsioFile.cpp b/Engine/Source/Developer/TraceAnalysis/Private/Store/AsioFile.cpp
--- a/Engine/Source/Developer/TraceAnalysis/Private/Store/AsioFile.cpp
+++ b/Engine/Source/Developer/TraceAnalysis/Private/Store/AsioFile.cpp
@@ -2,8 +2,6 @@
 
 #include "AsioFile.h"
 
-#if TRACE_WITH_ASIO
-
 #if !PLATFORM_WINDOWS
 #include <sys/file.h>
 #include <errno.h>
@@ -133,5 +131,3 @@ FAsioReadable* FAsioFile::ReadFile(asio::io_context& IoContext, const TCHAR* Pat
 }
 
 } // namespace Trace
-
-#endif // TRACE_WITH_ASIO
diff --git a/Engine/Source/Developer/TraceAnalysis/Private/Store/AsioFile.h b/Engine/Source/Developer/TraceAnalysis/Private/Store/AsioFile.h
--- a/Engine/Source/Developer/TraceAnalysis/Private/Store/AsioFile.h
+++ b/Engine/Source/Developer/TraceAnalysis/Private/Store/AsioFile.h
@@ -4,8 +4,6 @@
 
 #include "Asio/Asio.h"
 
-#if TRACE_WITH_ASIO
-
 #include "AsioIoable.h"
 #if !PLATFORM_WINDOWS
 #include "asio/posix/stream_descriptor.hpp"
@@ -38,5 +36,3 @@ private:
 };
 
 } // namespace Trace
-
-#endif // TRACE_WITH_ASIO
diff --git a/Engine/Source/Developer/TraceAnalysis/Private/Store/AsioIoable.cpp b/Engine/Source/Developer/TraceAnalysis/Private/Store/AsioIoable.cpp
--- a/Engine/Source/Developer/TraceAnalysis/Private/Store/AsioIoable.cpp
+++ b/Engine/Source/Developer/TraceAnalysis/Private/Store/AsioIoable.cpp
@@ -2,8 +2,6 @@
 
 #include "AsioIoable.h"
 
-#if TRACE_WITH_ASIO
-
 namespace Trace
 {
 
@@ -46,5 +44,3 @@ void FAsioIoable::OnIoComplete(const asio::error_code& ErrorCode, int32 Size)
 }
 
 } // namespace Trace
-
-#endif // TRACE_WITH_ASIO
diff --git a/Engine/Source/Developer/TraceAnalysis/Private/Store/AsioIoable.h b/Engine/Source/Developer/TraceAnalysis/Private/Store/AsioIoable.h
--- a/Engine/Source/Developer/TraceAnalysis/Private/Store/AsioIoable.h
+++ b/Engine/Source/Developer/TraceAnalysis/Private/Store/AsioIoable.h
@@ -4,8 +4,6 @@
 
 #include "Asio/Asio.h"
 
-#if TRACE_WITH_ASIO
-
 namespace Trace
 {
 
@@ -52,5 +50,3 @@ public:
 };
 
 } // namespace Trace
-
-#endif // TRACE_WITH_ASIO
diff --git a/Engine/Source/Developer/TraceAnalysis/Private/Store/AsioObject.cpp b/Engine/Source/Developer/TraceAnalysis/Private/Store/AsioObject.cpp
--- a/Engine/Source/Developer/TraceAnalysis/Private/Store/AsioObject.cpp
+++ b/Engine/Source/Developer/TraceAnalysis/Private/Store/AsioObject.cpp
@@ -2,8 +2,6 @@
 
 #include "AsioObject.h"
 
-#if TRACE_WITH_ASIO
-
 namespace Trace
 {
 
@@ -20,5 +18,3 @@ asio::io_context& FAsioObject::GetIoContext()
 }
 
 } // namespace Trace
-
-#endif // TRACE_WITH_ASIO
diff --git a/Engine/Source/Developer/TraceAnalysis/Private/Store/AsioObject.h b/Engine/Source/Developer/TraceAnalysis/Private/Store/AsioObject.h
--- a/Engine/Source/Developer/TraceAnalysis/Private/Store/AsioObject.h
+++ b/Engine/Source/Developer/TraceAnalysis/Private/Store/AsioObject.h
@@ -4,8 +4,6 @@
 
 #include "Asio/Asio.h"
 
-#if TRACE_WITH_ASIO
-
 namespace Trace
 {
 
@@ -21,5 +19,3 @@ private:
 };
 
 } // namespace Trace
-
-#endif // TRACE_WITH_ASIO
diff --git a/Engine/Source/Developer/TraceAnalysis/Private/Store/AsioRecorder.cpp b/Engine/Source/Developer/TraceAnalysis/Private/Store/AsioRecorder.cpp
--- a/Engine/Source/Developer/TraceAnalysis/Private/Store/AsioRecorder.cpp
+++ b/Engine/Source/Developer/TraceAnalysis/Private/Store/AsioRecorder.cpp
@@ -1,9 +1,6 @@
 // Copyright Epic Games, Inc. All Rights Reserved.
 
 #include "AsioRecorder.h"
-
-#if TRACE_WITH_ASIO
-
 #include "AsioIoable.h"
 #include "AsioSocket.h"
 #include "AsioStore.h"
@@ -22,7 +19,6 @@ public:
 	void				Close();
 	uint32				GetIpAddress() const;
 
-
 private:
 	virtual void		OnIoComplete(uint32 Id, int32 Size) override;
 	static const uint32	BufferSize = 64 * 1024;
@@ -193,5 +189,3 @@ void FAsioRecorder::OnTick()
 }
 
 } // namespace Trace
-
-#endif // TRACE_WITH_ASIO
diff --git a/Engine/Source/Developer/TraceAnalysis/Private/Store/AsioRecorder.h b/Engine/Source/Developer/TraceAnalysis/Private/Store/AsioRecorder.h
--- a/Engine/Source/Developer/TraceAnalysis/Private/Store/AsioRecorder.h
+++ b/Engine/Source/Developer/TraceAnalysis/Private/Store/AsioRecorder.h
@@ -3,9 +3,6 @@
 #pragma once
 
 #include "Asio/Asio.h"
-
-#if TRACE_WITH_ASIO
-
 #include "AsioTcpServer.h"
 #include "AsioTickable.h"
 #include "Containers/Array.h"
@@ -49,5 +46,3 @@ private:
 };
 
 } // namespace Trace
-
-#endif // TRACE_WITH_ASIO
diff --git a/Engine/Source/Developer/TraceAnalysis/Private/Store/AsioSocket.cpp b/Engine/Source/Developer/TraceAnalysis/Private/Store/AsioSocket.cpp
--- a/Engine/Source/Developer/TraceAnalysis/Private/Store/AsioSocket.cpp
+++ b/Engine/Source/Developer/TraceAnalysis/Private/Store/AsioSocket.cpp
@@ -2,8 +2,6 @@
 
 #include "AsioSocket.h"
 
-#if TRACE_WITH_ASIO
-
 #include "Templates/UnrealTemplate.h"
 
 namespace Trace
@@ -79,7 +77,7 @@ bool FAsioSocket::Read(void* Dest, uint32 Size, FAsioIoSink* Sink, uint32 Id)
 
 	asio::async_read(
 		Socket,
-        asio::buffer(Dest, Size),
+		asio::buffer(Dest, Size),
 		[this] (const asio::error_code& ErrorCode, size_t BytesReceived)
 		{
 			OnIoComplete(ErrorCode, BytesReceived);
@@ -98,7 +96,7 @@ bool FAsioSocket::ReadSome(void* Dest, uint32 BufferSize, FAsioIoSink* Sink, uin
 	}
 
 	Socket.async_receive(
-        asio::buffer(Dest, BufferSize),
+		asio::buffer(Dest, BufferSize),
 		[this] (const asio::error_code& ErrorCode, size_t BytesReceived)
 		{
 			return OnIoComplete(ErrorCode, BytesReceived);
@@ -129,5 +127,3 @@ bool FAsioSocket::Write(const void* Src, uint32 Size, FAsioIoSink* Sink, uint32
 }
 
 } // namespace Trace
-
-#endif // TRACE_WITH_ASIO
diff --git a/Engine/Source/Developer/TraceAnalysis/Private/Store/AsioSocket.h b/Engine/Source/Developer/TraceAnalysis/Private/Store/AsioSocket.h
--- a/Engine/Source/Developer/TraceAnalysis/Private/Store/AsioSocket.h
+++ b/Engine/Source/Developer/TraceAnalysis/Private/Store/AsioSocket.h
@@ -3,9 +3,6 @@
 #pragma once
 
 #include "Asio/Asio.h"
-
-#if TRACE_WITH_ASIO
-
 #include "AsioIoable.h"
 
 namespace Trace
@@ -33,5 +30,3 @@ private:
 };
 
 } // namespace Trace
-
-#endif // TRACE_WITH_ASIO
diff --git a/Engine/Source/Developer/TraceAnalysis/Private/Store/AsioStore.cpp b/Engine/Source/Developer/TraceAnalysis/Private/Store/AsioStore.cpp
--- a/Engine/Source/Developer/TraceAnalysis/Private/Store/AsioStore.cpp
+++ b/Engine/Source/Developer/TraceAnalysis/Private/Store/AsioStore.cpp
@@ -1,9 +1,6 @@
 // Copyright Epic Games, Inc. All Rights Reserved.
 
 #include "AsioStore.h"
-
-#if TRACE_WITH_ASIO
-
 #include "HAL/PlatformFile.h"
 #include "Containers/UnrealString.h"
 #include "Misc/CString.h"
@@ -247,5 +244,3 @@ FAsioReadable* FAsioStore::OpenTrace(uint32 Id)
 }
 
 } // namespace Trace
-
-#endif // TRACE_WITH_ASIO
diff --git a/Engine/Source/Developer/TraceAnalysis/Private/Store/AsioStore.h b/Engine/Source/Developer/TraceAnalysis/Private/Store/AsioStore.h
--- a/Engine/Source/Developer/TraceAnalysis/Private/Store/AsioStore.h
+++ b/Engine/Source/Developer/TraceAnalysis/Private/Store/AsioStore.h
@@ -3,9 +3,6 @@
 #pragma once
 
 #include "Asio/Asio.h"
-
-#if TRACE_WITH_ASIO
-
 #include "AsioFile.h"
 #include "AsioObject.h"
 #include "Containers/Array.h"
@@ -69,5 +66,3 @@ private:
 };
 
 } // namespace Trace
-
-#endif // TRACE_WITH_ASIO
diff --git a/Engine/Source/Developer/TraceAnalysis/Private/Store/AsioStoreCborServer.cpp b/Engine/Source/Developer/TraceAnalysis/Private/Store/AsioStoreCborServer.cpp
--- a/Engine/Source/Developer/TraceAnalysis/Private/Store/AsioStoreCborServer.cpp
+++ b/Engine/Source/Developer/TraceAnalysis/Private/Store/AsioStoreCborServer.cpp
@@ -1,9 +1,6 @@
 // Copyright Epic Games, Inc. All Rights Reserved.
 
 #include "AsioStoreCborServer.h"
-
-#if TRACE_WITH_ASIO
-
 #include "AsioIoable.h"
 #include "AsioObject.h"
 #include "AsioRecorder.h"
@@ -408,5 +405,3 @@ void FAsioStoreCborServer::OnTick()
 }
 
 } // namespace Trace
-
-#endif // TRACE_WITH_ASIO
diff --git a/Engine/Source/Developer/TraceAnalysis/Private/Store/AsioStoreCborServer.h b/Engine/Source/Developer/TraceAnalysis/Private/Store/AsioStoreCborServer.h
--- a/Engine/Source/Developer/TraceAnalysis/Private/Store/AsioStoreCborServer.h
+++ b/Engine/Source/Developer/TraceAnalysis/Private/Store/AsioStoreCborServer.h
@@ -3,9 +3,6 @@
 #pragma once
 
 #include "Asio/Asio.h"
-
-#if TRACE_WITH_ASIO
-
 #include "AsioTcpServer.h"
 #include "AsioTickable.h"
 #include "Containers/Array.h"
@@ -37,5 +34,3 @@ private:
 };
 
 } // namespace Trace
-
-#endif // TRACE_WITH_ASIO
diff --git a/Engine/Source/Developer/TraceAnalysis/Private/Store/AsioTcpServer.cpp b/Engine/Source/Developer/TraceAnalysis/Private/Store/AsioTcpServer.cpp
--- a/Engine/Source/Developer/TraceAnalysis/Private/Store/AsioTcpServer.cpp
+++ b/Engine/Source/Developer/TraceAnalysis/Private/Store/AsioTcpServer.cpp
@@ -1,9 +1,6 @@
 // Copyright Epic Games, Inc. All Rights Reserved.
 
 #include "AsioTcpServer.h"
-
-#if TRACE_WITH_ASIO
-
 #include "Templates/UnrealTemplate.h"
 
 namespace Trace
@@ -51,7 +48,7 @@ bool FAsioTcpServer::StartServer(uint32 Port)
 
 	using asio::ip::tcp;
 
-    tcp::endpoint Endpoint(tcp::v4(), uint16(Port));
+	tcp::endpoint Endpoint(tcp::v4(), uint16(Port));
 	tcp::acceptor TempAcceptor(GetIoContext(), Endpoint, false);
 
 	asio::error_code ErrorCode;
@@ -102,5 +99,3 @@ void FAsioTcpServer::AsyncAccept()
 }
 
 } // namespace Trace
-
-#endif // TRACE_WITH_ASIO
diff --git a/Engine/Source/Developer/TraceAnalysis/Private/Store/AsioTcpServer.h b/Engine/Source/Developer/TraceAnalysis/Private/Store/AsioTcpServer.h
--- a/Engine/Source/Developer/TraceAnalysis/Private/Store/AsioTcpServer.h
+++ b/Engine/Source/Developer/TraceAnalysis/Private/Store/AsioTcpServer.h
@@ -3,9 +3,6 @@
 #pragma once
 
 #include "Asio/Asio.h"
-
-#if TRACE_WITH_ASIO
-
 #include "AsioObject.h"
 
 namespace Trace
@@ -29,5 +26,3 @@ private:
 };
 
 } // namespace Trace
-
-#endif // TRACE_WITH_ASIO
diff --git a/Engine/Source/Developer/TraceAnalysis/Private/Store/AsioTickable.cpp b/Engine/Source/Developer/TraceAnalysis/Private/Store/AsioTickable.cpp
--- a/Engine/Source/Developer/TraceAnalysis/Private/Store/AsioTickable.cpp
+++ b/Engine/Source/Developer/TraceAnalysis/Private/Store/AsioTickable.cpp
@@ -2,8 +2,6 @@
 
 #include "AsioTickable.h"
 
-#if TRACE_WITH_ASIO
-
 namespace Trace
 {
 
@@ -74,5 +72,3 @@ void FAsioTickable::AsyncTick()
 }
 
 } // namespace Trace
-
-#endif // TRACE_WITH_ASIO
diff --git a/Engine/Source/Developer/TraceAnalysis/Private/Store/AsioTickable.h b/Engine/Source/Developer/TraceAnalysis/Private/Store/AsioTickable.h
--- a/Engine/Source/Developer/TraceAnalysis/Private/Store/AsioTickable.h
+++ b/Engine/Source/Developer/TraceAnalysis/Private/Store/AsioTickable.h
@@ -3,9 +3,6 @@
 #pragma once
 
 #include "Asio/Asio.h"
-
-#if TRACE_WITH_ASIO
-
 #include "AsioObject.h"
 
 #include <chrono>
@@ -31,5 +28,3 @@ private:
 };
 
 } // namespace Trace
-
-#endif // TRACE_WITH_ASIO
diff --git a/Engine/Source/Developer/TraceAnalysis/Private/Store/AsioTraceRelay.cpp b/Engine/Source/Developer/TraceAnalysis/Private/Store/AsioTraceRelay.cpp
--- a/Engine/Source/Developer/TraceAnalysis/Private/Store/AsioTraceRelay.cpp
+++ b/Engine/Source/Developer/TraceAnalysis/Private/Store/AsioTraceRelay.cpp
@@ -1,9 +1,6 @@
 // Copyright Epic Games, Inc. All Rights Reserved.
 
 #include "AsioTraceRelay.h"
-
-#if TRACE_WITH_ASIO
-
 #include "AsioRecorder.h"
 #include "AsioSocket.h"
 
@@ -96,5 +93,3 @@ void FAsioTraceRelay::OnIoComplete(uint32 Id, int32 Size)
 }
 
 } // namespace Trace
-
-#endif // TRACE_WITH_ASIO
diff --git a/Engine/Source/Developer/TraceAnalysis/Private/Store/AsioTraceRelay.h b/Engine/Source/Developer/TraceAnalysis/Private/Store/AsioTraceRelay.h
--- a/Engine/Source/Developer/TraceAnalysis/Private/Store/AsioTraceRelay.h
+++ b/Engine/Source/Developer/TraceAnalysis/Private/Store/AsioTraceRelay.h
@@ -3,9 +3,6 @@
 #pragma once
 
 #include "Asio/Asio.h"
-
-#if TRACE_WITH_ASIO
-
 #include "AsioIoable.h"
 #include "AsioTcpServer.h"
 
@@ -39,5 +36,3 @@ private:
 };
 
 } // namespace Trace
-
-#endif // TRACE_WITH_ASIO
diff --git a/Engine/Source/Developer/TraceAnalysis/Private/Store/CborPayload.h b/Engine/Source/Developer/TraceAnalysis/Private/Store/CborPayload.h
--- a/Engine/Source/Developer/TraceAnalysis/Private/Store/CborPayload.h
+++ b/Engine/Source/Developer/TraceAnalysis/Private/Store/CborPayload.h
@@ -3,9 +3,6 @@
 #pragma once
 
 #include "Asio/Asio.h"
-
-#if TRACE_WITH_ASIO
-
 #include "CborReader.h"
 #include "CborWriter.h"
 #include "Containers/Array.h"
@@ -229,5 +226,3 @@ inline FAnsiStringView FResponse::GetString(const char* Key, const char (&Defaul
 }
 
 } // namespace Trace
-
-#endif // TRACE_WITH_ASIO
diff --git a/Engine/Source/Developer/TraceAnalysis/Private/Store/FileStore.cpp b/Engine/Source/Developer/TraceAnalysis/Private/Store/FileStore.cpp
--- a/Engine/Source/Developer/TraceAnalysis/Private/Store/FileStore.cpp
+++ b/Engine/Source/Developer/TraceAnalysis/Private/Store/FileStore.cpp
@@ -371,5 +371,4 @@ TSharedPtr<IStore> Store_Create(const TCHAR* StoreDir)
 	return FileStore;
 }
 
-}
-
+} // namespace Trace
diff --git a/Engine/Source/Developer/TraceAnalysis/Private/Store/StoreClient.cpp b/Engine/Source/Developer/TraceAnalysis/Private/Store/StoreClient.cpp
--- a/Engine/Source/Developer/TraceAnalysis/Private/Store/StoreClient.cpp
+++ b/Engine/Source/Developer/TraceAnalysis/Private/Store/StoreClient.cpp
@@ -2,9 +2,6 @@
 
 #include "Trace/StoreClient.h"
 #include "Asio/Asio.h"
-
-#if TRACE_WITH_ASIO
-
 #include "AsioStore.h"
 #include "CborPayload.h"
 #include "Templates/UnrealTemplate.h"
@@ -22,9 +19,7 @@ public:
 	virtual int32			Read(void* Dest, uint32 DestSize) override;
 
 private:
-#if TRACE_WITH_ASIO
 	asio::ip::tcp::socket	Socket;
-#endif // TRACE_WITH_ASIO
 };
 
 ////////////////////////////////////////////////////////////////////////////////
@@ -36,15 +31,12 @@ FTraceDataStream::FTraceDataStream(asio::ip::tcp::socket& InSocket)
 ////////////////////////////////////////////////////////////////////////////////
 FTraceDataStream::~FTraceDataStream()
 {
-#if TRACE_WITH_ASIO
 	Socket.close();
-#endif // TRACE_WITH_ASIO
 }
 
 ////////////////////////////////////////////////////////////////////////////////
 int32 FTraceDataStream::Read(void* Dest, uint32 DestSize)
 {
-#if TRACE_WITH_ASIO
 	asio::error_code ErrorCode;
 	size_t BytesRead = Socket.read_some(asio::buffer(Dest, DestSize), ErrorCode);
 	if (ErrorCode)
@@ -54,9 +46,6 @@ int32 FTraceDataStream::Read(void* Dest, uint32 DestSize)
 	}
 
 	return int32(BytesRead);
-#else
-	return -1;
-#endif // TRACE_WITH_ASIO
 }
 
 
@@ -309,19 +298,6 @@ bool FStoreCborClient::GetSessionInfoByTraceId(uint32 TraceId)
 
 } // namespace Trace
 
-#endif // TRACE_WITH_ASIO
-
-
-
-#if !TRACE_WITH_ASIO
-struct FResponse
-{
-	int64		GetInteger(const char*, int64 Default) const		{ return Default; }
-	const char* GetString(const char*, const char* Default) const	{ return Default; }
-};
-#endif // !TRACE_WITH_ASIO
-
-
 
 namespace Trace
 {
@@ -391,7 +367,6 @@ uint32 FStoreClient::FSessionInfo::GetIpAddress() const
 ////////////////////////////////////////////////////////////////////////////////
 FStoreClient* FStoreClient::Connect(const TCHAR* Host, uint32 Port)
 {
-#if TRACE_WITH_ASIO
 	FStoreCborClient* Impl = new FStoreCborClient();
 	if (!Impl->Connect(Host, Port))
 	{
@@ -400,57 +375,39 @@ FStoreClient* FStoreClient::Connect(const TCHAR* Host, uint32 Port)
 	}
 
 	return (FStoreClient*)Impl;
-#else
-	return nullptr;
-#endif
 }
 
 ////////////////////////////////////////////////////////////////////////////////
 void FStoreClient::operator delete (void* Addr)
 {
-#if TRACE_WITH_ASIO
 	auto* Self = (FStoreCborClient*)Addr;
 	delete Self;
-#endif
 }
 
 ////////////////////////////////////////////////////////////////////////////////
 bool FStoreClient::IsValid() const
 {
-#if TRACE_WITH_ASIO
 	auto* Self = (FStoreCborClient*)this;
 	return Self->IsOpen();
-#else
-	return false;
-#endif
 }
 
 ////////////////////////////////////////////////////////////////////////////////
 uint32 FStoreClient::GetStoreAddress() const
 {
-#if TRACE_WITH_ASIO
 	auto* Self = (FStoreCborClient*)this;
 	return Self->GetStoreAddress();
-#else
-	return 0;
-#endif
 }
 
 ////////////////////////////////////////////////////////////////////////////////
 uint32 FStoreClient::GetStorePort() const
 {
-#if TRACE_WITH_ASIO
 	auto* Self = (FStoreCborClient*)this;
 	return Self->GetStorePort();
-#else
-	return 0;
-#endif
 }
 
 ////////////////////////////////////////////////////////////////////////////////
 const FStoreClient::FStatus* FStoreClient::GetStatus()
 {
-#if TRACE_WITH_ASIO
 	auto* Self = (FStoreCborClient*)this;
 	if (!Self->GetStatus())
 	{
@@ -459,15 +416,11 @@ const FStoreClient::FStatus* FStoreClient::GetStatus()
 
 	const FResponse& Response = Self->GetResponse();
 	return (FStatus*)(&Response);
-#else
-	return nullptr;
-#endif // TRACE_WITH_ASIO
 }
 
 ////////////////////////////////////////////////////////////////////////////////
 uint32 FStoreClient::GetTraceCount()
 {
-#if TRACE_WITH_ASIO
 	auto* Self = (FStoreCborClient*)this;
 	if (!Self->GetTraceCount())
 	{
@@ -475,15 +428,11 @@ uint32 FStoreClient::GetTraceCount()
 	}
 
 	return Self->GetResponse().GetInteger("count", 0);
-#else
-	return 0;
-#endif // TRACE_WITH_ASIO
 }
 
 ////////////////////////////////////////////////////////////////////////////////
 const FStoreClient::FTraceInfo* FStoreClient::GetTraceInfo(uint32 Index)
 {
-#if TRACE_WITH_ASIO
 	auto* Self = (FStoreCborClient*)this;
 	if (!Self->GetTraceInfo(Index))
 	{
@@ -492,15 +441,11 @@ const FStoreClient::FTraceInfo* FStoreClient::GetTraceInfo(uint32 Index)
 
 	const FResponse& Response = Self->GetResponse();
 	return (FTraceInfo*)(&Response);
-#else
-	return nullptr;
-#endif // TRACE_WITH_ASIO
 }
 
 ////////////////////////////////////////////////////////////////////////////////
 const FStoreClient::FTraceInfo* FStoreClient::GetTraceInfoById(uint32 Id)
 {
-#if TRACE_WITH_ASIO
 	auto* Self = (FStoreCborClient*)this;
 	if (!Self->GetTraceInfoById(Id))
 	{
@@ -509,26 +454,18 @@ const FStoreClient::FTraceInfo* FStoreClient::GetTraceInfoById(uint32 Id)
 
 	const FResponse& Response = Self->GetResponse();
 	return (FTraceInfo*)(&Response);
-#else
-	return nullptr;
-#endif // TRACE_WITH_ASIO
 }
 
 ////////////////////////////////////////////////////////////////////////////////
 FStoreClient::FTraceData FStoreClient::ReadTrace(uint32 Id)
 {
-#if TRACE_WITH_ASIO
 	auto* Self = (FStoreCborClient*)this;
 	return FTraceData(Self->ReadTrace(Id));
-#else
-	return nullptr;
-#endif // TRACE_WITH_ASIO
 }
 
 ////////////////////////////////////////////////////////////////////////////////
 uint32 FStoreClient::GetSessionCount() const
 {
-#if TRACE_WITH_ASIO
 	auto* Self = (FStoreCborClient*)this;
 	if (!Self->GetSessionCount())
 	{
@@ -536,15 +473,11 @@ uint32 FStoreClient::GetSessionCount() const
 	}
 
 	return Self->GetResponse().GetInteger("count", 0);
-#else
-	return 0;
-#endif // TRACE_WITH_ASIO
 }
 
 ////////////////////////////////////////////////////////////////////////////////
 const FStoreClient::FSessionInfo* FStoreClient::GetSessionInfo(uint32 Index) const
 {
-#if TRACE_WITH_ASIO
 	auto* Self = (FStoreCborClient*)this;
 	if (!Self->GetSessionInfo(Index))
 	{
@@ -553,15 +486,11 @@ const FStoreClient::FSessionInfo* FStoreClient::GetSessionInfo(uint32 Index) con
 
 	const FResponse& Response = Self->GetResponse();
 	return (FSessionInfo*)(&Response);
-#else
-	return nullptr;
-#endif // TRACE_WITH_ASIO
 }
 
 ////////////////////////////////////////////////////////////////////////////////
 const FStoreClient::FSessionInfo* FStoreClient::GetSessionInfoById(uint32 Id) const
 {
-#if TRACE_WITH_ASIO
 	auto* Self = (FStoreCborClient*)this;
 	if (!Self->GetSessionInfoById(Id))
 	{
@@ -570,15 +499,11 @@ const FStoreClient::FSessionInfo* FStoreClient::GetSessionInfoById(uint32 Id) co
 
 	const FResponse& Response = Self->GetResponse();
 	return (FSessionInfo*)(&Response);
-#else
-	return nullptr;
-#endif // TRACE_WITH_ASIO
 }
 
 ////////////////////////////////////////////////////////////////////////////////
 const FStoreClient::FSessionInfo* FStoreClient::GetSessionInfoByTraceId(uint32 TraceId) const
 {
-#if TRACE_WITH_ASIO
 	auto* Self = (FStoreCborClient*)this;
 	if (!Self->GetSessionInfoByTraceId(TraceId))
 	{
@@ -587,9 +512,6 @@ const FStoreClient::FSessionInfo* FStoreClient::GetSessionInfoByTraceId(uint32 T
 
 	const FResponse& Response = Self->GetResponse();
 	return (FSessionInfo*)(&Response);
-#else
-	return nullptr;
-#endif // TRACE_WITH_ASIO
 }
 
 } // namespace Trace
diff --git a/Engine/Source/Developer/TraceAnalysis/Private/Store/StoreService.cpp b/Engine/Source/Developer/TraceAnalysis/Private/Store/StoreService.cpp
--- a/Engine/Source/Developer/TraceAnalysis/Private/Store/StoreService.cpp
+++ b/Engine/Source/Developer/TraceAnalysis/Private/Store/StoreService.cpp
@@ -2,9 +2,6 @@
 
 #include "Trace/StoreService.h"
 #include "Asio/Asio.h"
-
-#if TRACE_WITH_ASIO
-
 #include "AsioContext.h"
 #include "AsioRecorder.h"
 #include "AsioStore.h"
@@ -41,19 +38,9 @@ FStoreServiceImpl::~FStoreServiceImpl()
 	Context.Stop();
 }
 
-} // namespace Trace
-
-#endif // TRACE_WITH_ASIO
-
-
-
-namespace Trace
-{
-
 ////////////////////////////////////////////////////////////////////////////////
 FStoreService* FStoreService::Create(const FDesc& InDesc)
 {
-#if TRACE_WITH_ASIO
 	FDesc Desc = InDesc;
 
 	IPlatformFile& PlatformFile = IPlatformFile::GetPlatformPhysical();
@@ -77,29 +64,20 @@ FStoreService* FStoreService::Create(const FDesc& InDesc)
 	Impl->Context.Start();
 
 	return (FStoreService*)Impl;
-#else
-	return nullptr;
-#endif // TRACE_WITH_ASIO
 }
 
 ////////////////////////////////////////////////////////////////////////////////
 void FStoreService::operator delete (void* Addr)
 {
-#if TRACE_WITH_ASIO
 	auto* Self = (FStoreServiceImpl*)Addr;
 	delete Self;
-#endif // TRACE_WITH_ASIO
 }
 
 ////////////////////////////////////////////////////////////////////////////////
 uint32 FStoreService::GetPort() const
 {
-#if TRACE_WITH_ASIO
 	auto* Self = (FStoreServiceImpl*)this;
 	return Self->CborServer.GetPort();
-#else
-	return 0;
-#endif // TRACE_WITH_ASIO
 }
 
 } // namespace Trace
