Fixed an ununitynopch compile error.

#rnx

diff --git a/Engine/Source/Developer/TraceAnalysis/Private/Store/AsioFile.cpp b/Engine/Source/Developer/TraceAnalysis/Private/Store/AsioFile.cpp
--- a/Engine/Source/Developer/TraceAnalysis/Private/Store/AsioFile.cpp
+++ b/Engine/Source/Developer/TraceAnalysis/Private/Store/AsioFile.cpp
@@ -3,6 +3,8 @@
 #include "AsioFile.h"
 
 #if !PLATFORM_WINDOWS
+#include "Containers/StringConv.h"
+
 #include <sys/file.h>
 #include <errno.h>
 #endif
