Added a method to tick an ASIO tickable just the one once in some amount of milliseconds time.

#rnx

diff --git a/Engine/Source/Developer/TraceAnalysis/Private/Store/AsioTickable.cpp b/Engine/Source/Developer/TraceAnalysis/Private/Store/AsioTickable.cpp
--- a/Engine/Source/Developer/TraceAnalysis/Private/Store/AsioTickable.cpp
+++ b/Engine/Source/Developer/TraceAnalysis/Private/Store/AsioTickable.cpp
@@ -49,6 +49,22 @@ bool FAsioTickable::StopTick()
 	return true;
 }
 
+////////////////////////////////////////////////////////////////////////////////
+void FAsioTickable::TickOnce(uint32 InMillisecondRate)
+{
+	if (MillisecondRate)
+	{
+		return;
+	}
+
+	MillisecondRate = InMillisecondRate;
+	if (MillisecondRate)
+	{
+		AsyncTick();
+		MillisecondRate = 0;
+	}
+}
+
 ////////////////////////////////////////////////////////////////////////////////
 void FAsioTickable::AsyncTick()
 {
diff --git a/Engine/Source/Developer/TraceAnalysis/Private/Store/AsioTickable.h b/Engine/Source/Developer/TraceAnalysis/Private/Store/AsioTickable.h
--- a/Engine/Source/Developer/TraceAnalysis/Private/Store/AsioTickable.h
+++ b/Engine/Source/Developer/TraceAnalysis/Private/Store/AsioTickable.h
@@ -19,6 +19,7 @@ public:
 	virtual				~FAsioTickable();
 	bool				StartTick(uint32 MillisecondRate);
 	bool				StopTick();
+	void				TickOnce(uint32 MillisecondRate);
 	virtual void		OnTick() = 0;
 
 private:
