Fixed Linux -disableunity/-nopch compile error

#rnx

diff --git a/Engine/Source/Developer/TraceAnalysis/Private/Store/AsioStore.cpp b/Engine/Source/Developer/TraceAnalysis/Private/Store/AsioStore.cpp
--- a/Engine/Source/Developer/TraceAnalysis/Private/Store/AsioStore.cpp
+++ b/Engine/Source/Developer/TraceAnalysis/Private/Store/AsioStore.cpp
@@ -2,6 +2,7 @@
 
 #include "AsioStore.h"
 #include "HAL/PlatformFile.h"
+#include "Containers/StringConv.h"
 #include "Containers/UnrealString.h"
 #include "Misc/CString.h"
 #include "Misc/DateTime.h"
