Use ASIO to wait when relaying a trace rather than blocking a ASIO worker thread in a Sleep().

#rnx

diff --git a/Engine/Source/Developer/TraceAnalysis/Private/Store/AsioTraceRelay.cpp b/Engine/Source/Developer/TraceAnalysis/Private/Store/AsioTraceRelay.cpp
--- a/Engine/Source/Developer/TraceAnalysis/Private/Store/AsioTraceRelay.cpp
+++ b/Engine/Source/Developer/TraceAnalysis/Private/Store/AsioTraceRelay.cpp
@@ -14,6 +14,7 @@ FAsioTraceRelay::FAsioTraceRelay(
 	uint32 InSessionId,
 	FAsioRecorder& InRecorder)
 : FAsioTcpServer(IoContext)
+, FAsioTickable(IoContext)
 , Input(InInput)
 , Recorder(InRecorder)
 , SessionId(InSessionId)
@@ -56,20 +57,25 @@ bool FAsioTraceRelay::OnAccept(asio::ip::tcp::socket& Socket)
 	return false;
 }
 
+////////////////////////////////////////////////////////////////////////////////
+void FAsioTraceRelay::OnTick()
+{
+	return OnIoComplete(OpStart, 0);
+}
+
 ////////////////////////////////////////////////////////////////////////////////
 void FAsioTraceRelay::OnIoComplete(uint32 Id, int32 Size)
 {
 	if (Size < 0)
 	{
-		if (Id == OpRead && SessionId)
+		if (Id == OpRead && SessionId /*&& -Size == asio::eof*/)
 		{
 			for (int i = 0, n = Recorder.GetSessionCount(); i < n; ++i)
 			{
 				const FAsioRecorder::FSession* Session = Recorder.GetSessionInfo(i);
 				if (Session->GetId() == SessionId)
 				{
-					FPlatformProcess::SleepNoStats(0.2f);
-					return OnIoComplete(OpStart, 0);
+					return TickOnce(200);
 				}
 			}
 		}
diff --git a/Engine/Source/Developer/TraceAnalysis/Private/Store/AsioTraceRelay.h b/Engine/Source/Developer/TraceAnalysis/Private/Store/AsioTraceRelay.h
--- a/Engine/Source/Developer/TraceAnalysis/Private/Store/AsioTraceRelay.h
+++ b/Engine/Source/Developer/TraceAnalysis/Private/Store/AsioTraceRelay.h
@@ -5,6 +5,7 @@
 #include "Asio/Asio.h"
 #include "AsioIoable.h"
 #include "AsioTcpServer.h"
+#include "AsioTickable.h"
 
 namespace Trace
 {
@@ -16,6 +17,7 @@ class FAsioSocket;
 class FAsioTraceRelay
 	: public FAsioTcpServer
 	, public FAsioIoSink
+	, public FAsioTickable
 {
 public:
 						FAsioTraceRelay(asio::io_context& IoContext, FAsioReadable* InInput, uint32 InTraceIf, FAsioRecorder& InRecorder);
@@ -25,6 +27,7 @@ public:
 
 private:
 	virtual bool		OnAccept(asio::ip::tcp::socket& Socket) override;
+	virtual void		OnTick() override;
 	virtual void		OnIoComplete(uint32 Id, int32 Size) override;
 	enum				{ OpStart, OpRead, OpSend };
 	static const uint32	BufferSize = 64 << 10;
