Expose recorder port in store.

diff --git a/Engine/Source/Developer/TraceAnalysis/Private/Store/StoreService.cpp b/Engine/Source/Developer/TraceAnalysis/Private/Store/StoreService.cpp
--- a/Engine/Source/Developer/TraceAnalysis/Private/Store/StoreService.cpp
+++ b/Engine/Source/Developer/TraceAnalysis/Private/Store/StoreService.cpp
@@ -80,4 +80,11 @@ uint32 FStoreService::GetPort() const
 	return Self->CborServer.GetPort();
 }
 
+////////////////////////////////////////////////////////////////////////////////
+uint32 FStoreService::GetRecorderPort() const
+{
+	auto* Self = (FStoreServiceImpl*)this;
+	return Self->Recorder.GetPort();
+}
+
 } // namespace Trace
diff --git a/Engine/Source/Developer/TraceAnalysis/Public/Trace/StoreService.h b/Engine/Source/Developer/TraceAnalysis/Public/Trace/StoreService.h
--- a/Engine/Source/Developer/TraceAnalysis/Public/Trace/StoreService.h
+++ b/Engine/Source/Developer/TraceAnalysis/Public/Trace/StoreService.h
@@ -22,6 +22,7 @@ public:
 	static FStoreService*	Create(const FDesc& Desc);
 	void					operator delete (void* Addr);
 	uint32					GetPort() const;
+	uint32					GetRecorderPort() const;
 
 private:
 							FStoreService() = default;
