Best practice is to call shutdown on a socket prior to closing it so the OS ensures all written data is sent.

#rnx

diff --git a/Engine/Source/Developer/TraceAnalysis/Private/Store/AsioSocket.cpp b/Engine/Source/Developer/TraceAnalysis/Private/Store/AsioSocket.cpp
--- a/Engine/Source/Developer/TraceAnalysis/Private/Store/AsioSocket.cpp
+++ b/Engine/Source/Developer/TraceAnalysis/Private/Store/AsioSocket.cpp
@@ -64,6 +64,7 @@ uint32 FAsioSocket::GetLocalPort() const
 ////////////////////////////////////////////////////////////////////////////////
 void FAsioSocket::Close()
 {
+	Socket.shutdown(asio::ip::tcp::socket::shutdown_both);
 	Socket.close();
 }
 
