DataStream derivers can implement a Close() method that can be used to terminate any consumers of the stream.

#rnx

diff --git a/Engine/Source/Developer/TraceAnalysis/Private/Store/StoreClient.cpp b/Engine/Source/Developer/TraceAnalysis/Private/Store/StoreClient.cpp
--- a/Engine/Source/Developer/TraceAnalysis/Private/Store/StoreClient.cpp
+++ b/Engine/Source/Developer/TraceAnalysis/Private/Store/StoreClient.cpp
@@ -17,6 +17,7 @@ public:
 							FTraceDataStream(asio::ip::tcp::socket& InSocket);
 	virtual					~FTraceDataStream();
 	virtual int32			Read(void* Dest, uint32 DestSize) override;
+	virtual void			Close() override;
 
 private:
 	asio::ip::tcp::socket	Socket;
@@ -31,6 +32,18 @@ FTraceDataStream::FTraceDataStream(asio::ip::tcp::socket& InSocket)
 ////////////////////////////////////////////////////////////////////////////////
 FTraceDataStream::~FTraceDataStream()
 {
+	Close();
+}
+
+////////////////////////////////////////////////////////////////////////////////
+void FTraceDataStream::Close()
+{
+	if (!Socket.is_open())
+	{
+		return;
+	}
+
+	Socket.shutdown(asio::ip::tcp::socket::shutdown_both);
 	Socket.close();
 }
 
@@ -258,6 +271,12 @@ FTraceDataStream* FStoreCborClient::ReadTrace(uint32 Id)
 		return nullptr;
 	}
 
+	SenderSocket.native_non_blocking(false, ErrorCode);
+	if (ErrorCode)
+	{
+		return nullptr;
+	}
+
 	return new FTraceDataStream(SenderSocket);
 }
 
diff --git a/Engine/Source/Developer/TraceAnalysis/Public/Trace/DataStream.h b/Engine/Source/Developer/TraceAnalysis/Public/Trace/DataStream.h
--- a/Engine/Source/Developer/TraceAnalysis/Public/Trace/DataStream.h
+++ b/Engine/Source/Developer/TraceAnalysis/Public/Trace/DataStream.h
@@ -12,6 +12,7 @@ class IInDataStream
 public:
 	virtual			~IInDataStream() = default;
 	virtual int32	Read(void* Data, uint32 Size) = 0;
+	virtual void	Close() {}
 };
 
 class IOutDataStream
