Explicitly shutdown and close store clients' sockets on destruction.

#rnx

diff --git a/Engine/Source/Developer/TraceAnalysis/Private/Store/StoreClient.cpp b/Engine/Source/Developer/TraceAnalysis/Private/Store/StoreClient.cpp
--- a/Engine/Source/Developer/TraceAnalysis/Private/Store/StoreClient.cpp
+++ b/Engine/Source/Developer/TraceAnalysis/Private/Store/StoreClient.cpp
@@ -100,6 +100,8 @@ FStoreCborClient::FStoreCborClient()
 ////////////////////////////////////////////////////////////////////////////////
 FStoreCborClient::~FStoreCborClient()
 {
+	Socket.shutdown(asio::ip::tcp::socket::shutdown_both);
+	Socket.close();
 }
 
 ////////////////////////////////////////////////////////////////////////////////
