Store the full path to a trace in FAsioStore::FTrace

#rnx

diff --git a/Engine/Source/Developer/TraceAnalysis/Private/Store/AsioStore.cpp b/Engine/Source/Developer/TraceAnalysis/Private/Store/AsioStore.cpp
--- a/Engine/Source/Developer/TraceAnalysis/Private/Store/AsioStore.cpp
+++ b/Engine/Source/Developer/TraceAnalysis/Private/Store/AsioStore.cpp
@@ -11,9 +11,27 @@ namespace Trace
 {
 
 ////////////////////////////////////////////////////////////////////////////////
-const TCHAR* FAsioStore::FTrace::GetName() const
+const FStringView& FAsioStore::FTrace::GetName() const
 {
-	return *Name;
+	if (Name.Len() == 0)
+	{
+		const TCHAR* Dot = FCString::Strrchr(*Path, '.');
+		if (Dot == nullptr)
+		{
+			Dot = *Path;
+		}
+
+		for (const TCHAR* c = Dot; c > *Path; --c)
+		{
+			if (c[-1] == '\\' || c[-1] == '/')
+			{
+				Name = FStringView(c, int32(Dot - c));
+				break;
+			}
+		}
+	}
+
+	return Name;
 }
 
 ////////////////////////////////////////////////////////////////////////////////
@@ -143,25 +161,9 @@ FAsioStore::FTrace* FAsioStore::AddTrace(const TCHAR* Path)
 	}
 #endif
 
-	const TCHAR* Dot = FCString::Strrchr(Path, '.');
-	if (Dot == nullptr)
-	{
-		return nullptr;
-	}
-
-	FStringView Name;
-	for (const TCHAR* c = Dot; c > Path; --c)
-	{
-		if (c[-1] == '\\' || c[-1] == '/')
-		{
-			Name = FStringView(c, int32(Dot - c));
-			break;
-		}
-	}
-
 	FTrace* Trace = new FTrace();
-	Trace->Name = Name;
-	Trace->Id = QuickStoreHash(Name);
+	Trace->Path = Path;
+	Trace->Id = QuickStoreHash(Trace->GetName());
 	Trace->Handle = UPTRINT(Handle);
 
 	Traces.Add(Trace);
@@ -238,7 +240,8 @@ FAsioReadable* FAsioStore::OpenTrace(uint32 Id)
 #if 0
 	TracePath.Appendf(TEXT("/%05d/data.utrace"), Id);
 #else
-	TracePath.Appendf(TEXT("/%s.utrace"), Trace->GetName());
+	const FStringView& Name = Trace->GetName();
+	TracePath.Appendf(TEXT("/%.*s.utrace"), Name.Len(), Name.GetData());
 #endif // 0
 
 	return FAsioFile::ReadFile(GetIoContext(), *TracePath);
diff --git a/Engine/Source/Developer/TraceAnalysis/Private/Store/AsioStore.h b/Engine/Source/Developer/TraceAnalysis/Private/Store/AsioStore.h
--- a/Engine/Source/Developer/TraceAnalysis/Private/Store/AsioStore.h
+++ b/Engine/Source/Developer/TraceAnalysis/Private/Store/AsioStore.h
@@ -6,6 +6,7 @@
 #include "AsioFile.h"
 #include "AsioObject.h"
 #include "Containers/Array.h"
+#include "Containers/StringView.h"
 #include "Containers/UnrealString.h"
 #include "Utils.h"
 
@@ -29,16 +30,18 @@ public:
 	class FTrace
 	{
 	public:
-		const TCHAR*	GetName() const;
-		uint32			GetId() const;
-		uint64			GetSize() const;
-		uint64			GetTimestamp() const;
+		const FStringView&	GetName() const;
+		uint32				GetId() const;
+		uint64				GetSize() const;
+		uint64				GetTimestamp() const;
 
 	private:
-		friend			FAsioStore;
-		FString			Name;
-		UPTRINT			Handle;
-		uint32			Id;
+		friend				FAsioStore;
+		FString				Path;
+		mutable FStringView	Name;
+		uint64				Timestamp;
+		UPTRINT				Handle;
+		uint32		        Id;
 	};
 
 	struct FNewTrace
diff --git a/Engine/Source/Developer/TraceAnalysis/Private/Store/AsioStoreCborServer.cpp b/Engine/Source/Developer/TraceAnalysis/Private/Store/AsioStoreCborServer.cpp
--- a/Engine/Source/Developer/TraceAnalysis/Private/Store/AsioStoreCborServer.cpp
+++ b/Engine/Source/Developer/TraceAnalysis/Private/Store/AsioStoreCborServer.cpp
@@ -196,23 +196,19 @@ void FAsioStoreCborPeer::OnTraceInfo()
 		return SendError(EStatusCode::BadRequest);
 	}
 
-	const TCHAR* Name = Trace->GetName();
 	char OutName[256];
-	for (char& Out : OutName)
+	const FStringView& Name = Trace->GetName();
+	uint32 NameLength = FMath::Min(uint32(sizeof(OutName)), uint32(Name.Len()));
+	for (uint32 i = 0; i < NameLength; ++i)
 	{
-		Out = char(*Name++);
-		if (Out == '\0')
-		{
-			break;
-		}
+		OutName[i] = char(Name[i]);
 	}
-	OutName[sizeof(OutName) - 1] = '\0';
 
 	TPayloadBuilder<> Builder((int32)EStatusCode::Success);
 	Builder.AddInteger("id", Trace->GetId());
 	Builder.AddInteger("size", Trace->GetSize());
 	Builder.AddInteger("timestamp", Trace->GetTimestamp());
-	Builder.AddString("name", OutName);
+	Builder.AddString("name", OutName, NameLength);
 	SendResponse(Builder.Done());
 }
 
