Build a valid FTrace in its contructor instead of AddTrace().

#rnx

diff --git a/Engine/Source/Developer/TraceAnalysis/Private/Store/AsioStore.cpp b/Engine/Source/Developer/TraceAnalysis/Private/Store/AsioStore.cpp
--- a/Engine/Source/Developer/TraceAnalysis/Private/Store/AsioStore.cpp
+++ b/Engine/Source/Developer/TraceAnalysis/Private/Store/AsioStore.cpp
@@ -11,26 +11,30 @@ namespace Trace
 {
 
 ////////////////////////////////////////////////////////////////////////////////
-const FStringView& FAsioStore::FTrace::GetName() const
+FAsioStore::FTrace::FTrace(const TCHAR* InPath)
+: Path(InPath)
+, Id(QuickStoreHash(InPath))
 {
-	if (Name.Len() == 0)
+	// Extract the trace's name
+	const TCHAR* Dot = FCString::Strrchr(*Path, '.');
+	if (Dot == nullptr)
 	{
-		const TCHAR* Dot = FCString::Strrchr(*Path, '.');
-		if (Dot == nullptr)
-		{
-			Dot = *Path;
-		}
+		Dot = *Path;
+	}
 
-		for (const TCHAR* c = Dot; c > *Path; --c)
+	for (const TCHAR* c = Dot; c > *Path; --c)
+	{
+		if (c[-1] == '\\' || c[-1] == '/')
 		{
-			if (c[-1] == '\\' || c[-1] == '/')
-			{
-				Name = FStringView(c, int32(Dot - c));
-				break;
-			}
+			Name = FStringView(c, int32(Dot - c));
+			break;
 		}
 	}
+}
 
+////////////////////////////////////////////////////////////////////////////////
+const FStringView& FAsioStore::FTrace::GetName() const
+{
 	return Name;
 }
 
@@ -161,9 +165,7 @@ FAsioStore::FTrace* FAsioStore::AddTrace(const TCHAR* Path)
 	}
 #endif
 
-	FTrace* Trace = new FTrace();
-	Trace->Path = Path;
-	Trace->Id = QuickStoreHash(Trace->GetName());
+	FTrace* Trace = new FTrace(Path);
 	Trace->Handle = UPTRINT(Handle);
 
 	Traces.Add(Trace);
diff --git a/Engine/Source/Developer/TraceAnalysis/Private/Store/AsioStore.h b/Engine/Source/Developer/TraceAnalysis/Private/Store/AsioStore.h
--- a/Engine/Source/Developer/TraceAnalysis/Private/Store/AsioStore.h
+++ b/Engine/Source/Developer/TraceAnalysis/Private/Store/AsioStore.h
@@ -30,6 +30,7 @@ public:
 	class FTrace
 	{
 	public:
+							FTrace(const TCHAR* InPath);
 		const FStringView&	GetName() const;
 		uint32				GetId() const;
 		uint64				GetSize() const;
@@ -38,10 +39,10 @@ public:
 	private:
 		friend				FAsioStore;
 		FString				Path;
-		mutable FStringView	Name;
+		FStringView			Name;
 		uint64				Timestamp;
 		UPTRINT				Handle;
-		uint32		        Id;
+		uint32				Id = 0;
 	};
 
 	struct FNewTrace
