Do not hold a file handle open for each trace as it is rarely required.

#rnx

diff --git a/Engine/Source/Developer/TraceAnalysis/Private/Store/AsioStore.cpp b/Engine/Source/Developer/TraceAnalysis/Private/Store/AsioStore.cpp
--- a/Engine/Source/Developer/TraceAnalysis/Private/Store/AsioStore.cpp
+++ b/Engine/Source/Developer/TraceAnalysis/Private/Store/AsioStore.cpp
@@ -30,6 +30,36 @@ FAsioStore::FTrace::FTrace(const TCHAR* InPath)
 			break;
 		}
 	}
+
+	// Calculate that trace's timestamp
+	uint64 InTimestamp = 0;
+#if PLATFORM_WINDOWS
+	HANDLE Handle = CreateFileW(InPath, 0, FILE_SHARE_READ|FILE_SHARE_WRITE|FILE_SHARE_DELETE,
+		nullptr, OPEN_EXISTING, 0, nullptr);
+	if (Handle != INVALID_HANDLE_VALUE)
+	{
+		FILETIME Time;
+		if (GetFileTime(Handle, &Time, nullptr, nullptr))
+		{
+			// Windows FILETIME is a 64-bit value that represents the number of 100-nanosecond intervals that have elapsed since 12:00 A.M.January 1, 1601 Coordinated Universal Time(UTC).
+			// We adjust it to be compatible with the FDateTime ticks number of 100-nanosecond intervals that have elapsed since 12:00 A.M.January 1, 0001 Coordinated Universal Time(UTC).
+			const uint64 WinTicks = (static_cast<uint64>(Time.dwHighDateTime) << 32ull) | static_cast<uint64>(Time.dwLowDateTime);
+			const uint64 Year1601 = 504911232000000000ULL; // FDateTime(1601, 1, 1).GetTicks()
+			InTimestamp = Year1601 + WinTicks;
+		}
+
+		CloseHandle(Handle);
+	}
+#else
+	struct stat FileStat;
+	if (stat(TCHAR_TO_UTF8(InPath), &FileStat) == 0)
+	{
+		InTimestamp = (uint64(FileStat.st_ctim.tv_sec) * 1000 * 1000 * 1000) + FileStat.st_ctim.tv_nsec;
+		InTimestamp /= 100;
+		InTimestamp += 0x004d5bd15e978000ull; // k = unix_epoch - FDateTime_epoch (in 100-nanoseconds)
+	}
+#endif
+	Timestamp = InTimestamp;
 }
 
 ////////////////////////////////////////////////////////////////////////////////
@@ -48,33 +78,29 @@ uint32 FAsioStore::FTrace::GetId() const
 uint64 FAsioStore::FTrace::GetSize() const
 {
 #if PLATFORM_WINDOWS
-	HANDLE Inner = HANDLE(Handle);
-	LARGE_INTEGER FileSize;
-	GetFileSizeEx(Inner, &FileSize);
+	LARGE_INTEGER FileSize = {};
+	HANDLE Handle = CreateFileW(*Path, 0, FILE_SHARE_READ|FILE_SHARE_WRITE|FILE_SHARE_DELETE,
+		nullptr, OPEN_EXISTING, 0, nullptr);
+	if (Handle != INVALID_HANDLE_VALUE)
+	{
+		GetFileSizeEx(Handle, &FileSize);
+		CloseHandle(Handle);
+	}
 	return FileSize.QuadPart;
 #else
-	return 0; // TODO
+	struct stat FileStat;
+	if (stat(TCHAR_TO_UTF8(*Path), &FileStat) == 0)
+	{
+		return uint64(FileStat.st_size);
+	}
+	return 0;
 #endif
 }
 
 ////////////////////////////////////////////////////////////////////////////////
 uint64 FAsioStore::FTrace::GetTimestamp() const
 {
-#if PLATFORM_WINDOWS
-	HANDLE Inner = HANDLE(Handle);
-	FILETIME Time;
-	if (!GetFileTime(Inner, &Time, nullptr, nullptr))
-	{
-		return 0;
-	}
-	// Windows FILETIME is a 64-bit value that represents the number of 100-nanosecond intervals that have elapsed since 12:00 A.M.January 1, 1601 Coordinated Universal Time(UTC).
-	// We adjust it to be compatible with the FDateTime ticks number of 100-nanosecond intervals that have elapsed since 12:00 A.M.January 1, 0001 Coordinated Universal Time(UTC).
-	const uint64 WinTicks = (static_cast<uint64>(Time.dwHighDateTime) << 32ull) | static_cast<uint64>(Time.dwLowDateTime);
-	const uint64 Year1601 = 504911232000000000ULL; // FDateTime(1601, 1, 1).GetTicks()
-	return Year1601 + WinTicks;
-#else
-	return 0; // TODO
-#endif
+	return Timestamp;
 }
 
 
@@ -150,23 +176,7 @@ FAsioStore::FTrace* FAsioStore::GetTrace(uint32 Id)
 ////////////////////////////////////////////////////////////////////////////////
 FAsioStore::FTrace* FAsioStore::AddTrace(const TCHAR* Path)
 {
-#if PLATFORM_WINDOWS
-	HANDLE Handle = CreateFileW(Path, 0, FILE_SHARE_READ|FILE_SHARE_WRITE|FILE_SHARE_DELETE,
-		nullptr, OPEN_EXISTING, 0, nullptr);
-	if (Handle == INVALID_HANDLE_VALUE)
-	{
-		return nullptr;
-	}
-#else
-	int Handle = open(TCHAR_TO_ANSI(Path), O_RDWR, 0666);
-	if (!Handle)
-	{
-		return nullptr;
-	}
-#endif
-
 	FTrace* Trace = new FTrace(Path);
-	Trace->Handle = UPTRINT(Handle);
 
 	Traces.Add(Trace);
 	return Trace;
diff --git a/Engine/Source/Developer/TraceAnalysis/Private/Store/AsioStore.h b/Engine/Source/Developer/TraceAnalysis/Private/Store/AsioStore.h
--- a/Engine/Source/Developer/TraceAnalysis/Private/Store/AsioStore.h
+++ b/Engine/Source/Developer/TraceAnalysis/Private/Store/AsioStore.h
@@ -41,7 +41,6 @@ public:
 		FString				Path;
 		FStringView			Name;
 		uint64				Timestamp;
-		UPTRINT				Handle;
 		uint32				Id = 0;
 	};
 
