Added a method to clear a store's traces

#rnx

diff --git a/Engine/Source/Developer/TraceAnalysis/Private/Store/AsioStore.cpp b/Engine/Source/Developer/TraceAnalysis/Private/Store/AsioStore.cpp
--- a/Engine/Source/Developer/TraceAnalysis/Private/Store/AsioStore.cpp
+++ b/Engine/Source/Developer/TraceAnalysis/Private/Store/AsioStore.cpp
@@ -135,11 +135,20 @@ FAsioStore::FAsioStore(asio::io_context& IoContext, const TCHAR* InStoreDir)
 
 ////////////////////////////////////////////////////////////////////////////////
 FAsioStore::~FAsioStore()
+{
+	ClearTraces();
+}
+
+////////////////////////////////////////////////////////////////////////////////
+void FAsioStore::ClearTraces()
 {
 	for (FTrace* Trace : Traces)
 	{
 		delete Trace;
 	}
+
+	Traces.Empty();
+}
 }
 
 ////////////////////////////////////////////////////////////////////////////////
diff --git a/Engine/Source/Developer/TraceAnalysis/Private/Store/AsioStore.h b/Engine/Source/Developer/TraceAnalysis/Private/Store/AsioStore.h
--- a/Engine/Source/Developer/TraceAnalysis/Private/Store/AsioStore.h
+++ b/Engine/Source/Developer/TraceAnalysis/Private/Store/AsioStore.h
@@ -61,6 +61,7 @@ public:
 private:
 	FTrace*				GetTrace(uint32 Id);
 	FTrace*				AddTrace(const TCHAR* Path);
+	void				ClearTraces();
 	FString				StoreDir;
 	TArray<FTrace*>		Traces;
 #if 0
