Added a method to refresh a store's list of traces

#rnx

diff --git a/Engine/Source/Developer/TraceAnalysis/Private/Store/AsioStore.cpp b/Engine/Source/Developer/TraceAnalysis/Private/Store/AsioStore.cpp
--- a/Engine/Source/Developer/TraceAnalysis/Private/Store/AsioStore.cpp
+++ b/Engine/Source/Developer/TraceAnalysis/Private/Store/AsioStore.cpp
@@ -110,27 +110,7 @@ FAsioStore::FAsioStore(asio::io_context& IoContext, const TCHAR* InStoreDir)
 : FAsioObject(IoContext)
 , StoreDir(InStoreDir)
 {
-	IPlatformFile& FileSystem = IPlatformFile::GetPlatformPhysical();
-	FileSystem.IterateDirectory(InStoreDir, [this] (const TCHAR* Path, bool IsDirectory)
-	{
-#if 0
-		if (IsDirectory)
-		{
-			int32 Id = FCString::Atoi(Path);
-			LastTraceId = (Id < LastTraceId) ? Id : LastTraceId;
-		}
-#else
-		const TCHAR* Dot = FCString::Strrchr(Path, '.');
-		if (Dot == nullptr || FCString::Strcmp(Dot, TEXT(".utrace")))
-		{
-			return true;
-		}
-
-		AddTrace(Path);
-#endif // 0
-
-		return true;
-	});
+	Refresh();
 }
 
 ////////////////////////////////////////////////////////////////////////////////
@@ -149,7 +129,6 @@ void FAsioStore::ClearTraces()
 
 	Traces.Empty();
 }
-}
 
 ////////////////////////////////////////////////////////////////////////////////
 uint32 FAsioStore::GetTraceCount() const
@@ -268,4 +247,39 @@ FAsioReadable* FAsioStore::OpenTrace(uint32 Id)
 	return FAsioFile::ReadFile(GetIoContext(), *TracePath);
 }
 
+////////////////////////////////////////////////////////////////////////////////
+void FAsioStore::Refresh()
+{
+	ClearTraces();
+
+	IPlatformFile& FileSystem = IPlatformFile::GetPlatformPhysical();
+	FileSystem.IterateDirectory(*StoreDir, [this] (const TCHAR* Path, bool IsDirectory)
+	{
+#if 0
+		if (!IsDirectory)
+		{
+			return true;
+		}
+
+		int32 Id = FCString::Atoi(Path);
+		LastTraceId = (Id < LastTraceId) ? Id : LastTraceId;
+#else
+		if (IsDirectory)
+		{
+			return true;
+		}
+
+		const TCHAR* Dot = FCString::Strrchr(Path, '.');
+		if (Dot == nullptr || FCString::Strcmp(Dot, TEXT(".utrace")))
+		{
+			return true;
+		}
+
+		AddTrace(Path);
+#endif // 0
+
+		return true;
+	});
+}
+
 } // namespace Trace
diff --git a/Engine/Source/Developer/TraceAnalysis/Private/Store/AsioStore.h b/Engine/Source/Developer/TraceAnalysis/Private/Store/AsioStore.h
--- a/Engine/Source/Developer/TraceAnalysis/Private/Store/AsioStore.h
+++ b/Engine/Source/Developer/TraceAnalysis/Private/Store/AsioStore.h
@@ -62,6 +62,7 @@ private:
 	FTrace*				GetTrace(uint32 Id);
 	FTrace*				AddTrace(const TCHAR* Path);
 	void				ClearTraces();
+	void				Refresh();
 	FString				StoreDir;
 	TArray<FTrace*>		Traces;
 #if 0
