Reuse traces when adding duplicates

#rnx

diff --git a/Engine/Source/Developer/TraceAnalysis/Private/Store/AsioStore.cpp b/Engine/Source/Developer/TraceAnalysis/Private/Store/AsioStore.cpp
--- a/Engine/Source/Developer/TraceAnalysis/Private/Store/AsioStore.cpp
+++ b/Engine/Source/Developer/TraceAnalysis/Private/Store/AsioStore.cpp
@@ -164,8 +164,15 @@ FAsioStore::FTrace* FAsioStore::GetTrace(uint32 Id)
 ////////////////////////////////////////////////////////////////////////////////
 FAsioStore::FTrace* FAsioStore::AddTrace(const TCHAR* Path)
 {
-	FTrace* Trace = new FTrace(Path);
+	FTrace NewTrace(Path);
 
+	uint32 Id = NewTrace.GetId();
+	if (FTrace* Existing = GetTrace(Id))
+	{
+		return Existing;
+	}
+
+	FTrace* Trace = new FTrace(MoveTemp(NewTrace));
 	Traces.Add(Trace);
 	return Trace;
 }
