Brought back directory watching so users can add .utrace folders to the appropriate directory and Insights will update accordingly.

#rnx

diff --git a/Engine/Source/Developer/TraceAnalysis/Private/Asio/Asio.h b/Engine/Source/Developer/TraceAnalysis/Private/Asio/Asio.h
--- a/Engine/Source/Developer/TraceAnalysis/Private/Asio/Asio.h
+++ b/Engine/Source/Developer/TraceAnalysis/Private/Asio/Asio.h
@@ -32,6 +32,7 @@ void throw_exception(const ExceptionType& Exception)
 
 #if PLATFORM_WINDOWS
 #	include "asio/windows/random_access_handle.hpp"
+#	include "asio/windows/object_handle.hpp"
 #endif
 
 THIRD_PARTY_INCLUDES_END
diff --git a/Engine/Source/Developer/TraceAnalysis/Private/Store/AsioStore.cpp b/Engine/Source/Developer/TraceAnalysis/Private/Store/AsioStore.cpp
--- a/Engine/Source/Developer/TraceAnalysis/Private/Store/AsioStore.cpp
+++ b/Engine/Source/Developer/TraceAnalysis/Private/Store/AsioStore.cpp
@@ -7,9 +7,34 @@
 #include "Misc/CString.h"
 #include "Misc/DateTime.h"
 
+#if PLATFORM_WINDOWS
+#	include "Windows/AllowWindowsPlatformTypes.h"
+#	include <Windows.h>
+#	include "Windows/HideWindowsPlatformTypes.h"
+#endif // PLATFORM_WINDOWS
+
 namespace Trace
 {
 
+////////////////////////////////////////////////////////////////////////////////
+#if PLATFORM_WINDOWS
+class FAsioStore::FDirWatcher
+	: public asio::windows::object_handle
+{
+public:
+	using asio::windows::object_handle::object_handle;
+};
+#else
+class FAsioStore::FDirWatcher
+	//: public asio::posix::stream_descriptor
+{
+public:
+	void async_wait(...) {}
+};
+#endif // PLATFORM_WINDOWS
+
+
+
 ////////////////////////////////////////////////////////////////////////////////
 FAsioStore::FTrace::FTrace(const TCHAR* InPath)
 : Path(InPath)
@@ -111,11 +136,23 @@ FAsioStore::FAsioStore(asio::io_context& IoContext, const TCHAR* InStoreDir)
 , StoreDir(InStoreDir)
 {
 	Refresh();
+
+#if PLATFORM_WINDOWS
+	HANDLE DirWatchHandle = FindFirstChangeNotificationW(InStoreDir, false, FILE_NOTIFY_CHANGE_FILE_NAME);
+	if (DirWatchHandle == INVALID_HANDLE_VALUE)
+	{
+		DirWatchHandle = 0;
+	}
+	DirWatcher = new FDirWatcher(IoContext, DirWatchHandle);
+#endif
+
+	WatchDir();
 }
 
 ////////////////////////////////////////////////////////////////////////////////
 FAsioStore::~FAsioStore()
 {
+	delete DirWatcher;
 	ClearTraces();
 }
 
@@ -130,6 +167,24 @@ void FAsioStore::ClearTraces()
 	Traces.Empty();
 }
 
+////////////////////////////////////////////////////////////////////////////////
+void FAsioStore::WatchDir()
+{
+	if (DirWatcher == nullptr)
+	{
+		return;
+	}
+	
+	DirWatcher->async_wait([this] (asio::error_code ErrorCode)
+	{
+#if PLATFORM_WINDOWS
+		FindNextChangeNotification(DirWatcher->native_handle());
+#endif
+		Refresh();
+		WatchDir();
+	});
+}
+
 ////////////////////////////////////////////////////////////////////////////////
 uint32 FAsioStore::GetTraceCount() const
 {
diff --git a/Engine/Source/Developer/TraceAnalysis/Private/Store/AsioStore.h b/Engine/Source/Developer/TraceAnalysis/Private/Store/AsioStore.h
--- a/Engine/Source/Developer/TraceAnalysis/Private/Store/AsioStore.h
+++ b/Engine/Source/Developer/TraceAnalysis/Private/Store/AsioStore.h
@@ -59,12 +59,15 @@ public:
 	FAsioReadable*		OpenTrace(uint32 Id);
 
 private:
+	class				FDirWatcher;
 	FTrace*				GetTrace(uint32 Id);
 	FTrace*				AddTrace(const TCHAR* Path);
 	void				ClearTraces();
+	void				WatchDir();
 	void				Refresh();
 	FString				StoreDir;
 	TArray<FTrace*>		Traces;
+	FDirWatcher*		DirWatcher = nullptr;
 #if 0
 	int32				LastTraceId = -1;
 #endif // 0
