Cleaner store service shutdown which closes everything and waits for the async IO loops to drain.

#rnx

diff --git a/Engine/Source/Developer/TraceAnalysis/Private/Store/AsioContext.cpp b/Engine/Source/Developer/TraceAnalysis/Private/Store/AsioContext.cpp
--- a/Engine/Source/Developer/TraceAnalysis/Private/Store/AsioContext.cpp
+++ b/Engine/Source/Developer/TraceAnalysis/Private/Store/AsioContext.cpp
@@ -21,7 +21,7 @@ FAsioContext::FAsioContext(int32 ThreadCount)
 ////////////////////////////////////////////////////////////////////////////////
 FAsioContext::~FAsioContext()
 {
-	Stop();
+	Wait();
 	delete IoContext;
 }
 
@@ -43,7 +43,7 @@ void FAsioContext::Start()
 }
 
 ////////////////////////////////////////////////////////////////////////////////
-void FAsioContext::Stop()
+void FAsioContext::Stop(bool bWait)
 {
 	if (!bRunning)
 	{
@@ -52,9 +52,21 @@ void FAsioContext::Stop()
 
 	IoContext->stop();
 
+	if (bWait)
+	{
+		Wait();
+	}
+}
+
+////////////////////////////////////////////////////////////////////////////////
+void FAsioContext::Wait()
+{
 	for (std::thread& Thread : ThreadPool)
 	{
-		Thread.join();
+		if (Thread.joinable())
+		{
+			Thread.join();
+		}
 	}
 
 	bRunning = false;
diff --git a/Engine/Source/Developer/TraceAnalysis/Private/Store/AsioContext.h b/Engine/Source/Developer/TraceAnalysis/Private/Store/AsioContext.h
--- a/Engine/Source/Developer/TraceAnalysis/Private/Store/AsioContext.h
+++ b/Engine/Source/Developer/TraceAnalysis/Private/Store/AsioContext.h
@@ -18,7 +18,8 @@ public:
 							~FAsioContext();
 	asio::io_context&		Get();
 	void					Start();
-	void					Stop();
+	void					Stop(bool bWait=false);
+	void					Wait();
 
 private:
 	asio::io_context*		IoContext;
diff --git a/Engine/Source/Developer/TraceAnalysis/Private/Store/AsioRecorder.cpp b/Engine/Source/Developer/TraceAnalysis/Private/Store/AsioRecorder.cpp
--- a/Engine/Source/Developer/TraceAnalysis/Private/Store/AsioRecorder.cpp
+++ b/Engine/Source/Developer/TraceAnalysis/Private/Store/AsioRecorder.cpp
@@ -118,10 +118,25 @@ FAsioRecorder::FAsioRecorder(asio::io_context& IoContext, FAsioStore& InStore)
 ////////////////////////////////////////////////////////////////////////////////
 FAsioRecorder::~FAsioRecorder()
 {
+	Close();
+
 	for (FSession& Session : Sessions)
 	{
 		delete Session.Relay;
 	}
+	Sessions.Empty();
+}
+
+////////////////////////////////////////////////////////////////////////////////
+void FAsioRecorder::Close()
+{
+	FAsioTickable::StopTick();
+	FAsioTcpServer::Close();
+
+	for (FSession& Session : Sessions)
+	{
+		Session.Relay->Close();
+	}
 }
 
 ////////////////////////////////////////////////////////////////////////////////
diff --git a/Engine/Source/Developer/TraceAnalysis/Private/Store/AsioRecorder.h b/Engine/Source/Developer/TraceAnalysis/Private/Store/AsioRecorder.h
--- a/Engine/Source/Developer/TraceAnalysis/Private/Store/AsioRecorder.h
+++ b/Engine/Source/Developer/TraceAnalysis/Private/Store/AsioRecorder.h
@@ -35,6 +35,7 @@ public:
 
 								FAsioRecorder(asio::io_context& IoContext, FAsioStore& InStore);
 								~FAsioRecorder();
+	void						Close();
 	uint32						GetSessionCount() const;
 	const FSession*				GetSessionInfo(uint32 Index) const;
 
diff --git a/Engine/Source/Developer/TraceAnalysis/Private/Store/AsioStore.cpp b/Engine/Source/Developer/TraceAnalysis/Private/Store/AsioStore.cpp
--- a/Engine/Source/Developer/TraceAnalysis/Private/Store/AsioStore.cpp
+++ b/Engine/Source/Developer/TraceAnalysis/Private/Store/AsioStore.cpp
@@ -30,6 +30,8 @@ class FAsioStore::FDirWatcher
 {
 public:
 	void async_wait(...) {}
+	void cancel() {}
+	void close() {}
 };
 #endif // PLATFORM_WINDOWS
 
@@ -152,7 +154,19 @@ FAsioStore::FAsioStore(asio::io_context& IoContext, const TCHAR* InStoreDir)
 ////////////////////////////////////////////////////////////////////////////////
 FAsioStore::~FAsioStore()
 {
+	Close();
 	delete DirWatcher;
+}
+
+////////////////////////////////////////////////////////////////////////////////
+void FAsioStore::Close()
+{
+	if (DirWatcher != nullptr)
+	{
+		DirWatcher->cancel();
+		DirWatcher->close();
+	}
+
 	ClearTraces();
 }
 
@@ -177,6 +191,11 @@ void FAsioStore::WatchDir()
 	
 	DirWatcher->async_wait([this] (asio::error_code ErrorCode)
 	{
+		if (ErrorCode)
+		{
+			return;
+		}
+
 #if PLATFORM_WINDOWS
 		FindNextChangeNotification(DirWatcher->native_handle());
 #endif
diff --git a/Engine/Source/Developer/TraceAnalysis/Private/Store/AsioStore.h b/Engine/Source/Developer/TraceAnalysis/Private/Store/AsioStore.h
--- a/Engine/Source/Developer/TraceAnalysis/Private/Store/AsioStore.h
+++ b/Engine/Source/Developer/TraceAnalysis/Private/Store/AsioStore.h
@@ -52,9 +52,9 @@ public:
 
 						FAsioStore(asio::io_context& IoContext, const TCHAR* InStoreDir);
 						~FAsioStore();
+	void				Close();
 	uint32				GetTraceCount() const;
 	const FTrace*		GetTraceInfo(uint32 Index) const;
-	int32				ReadTrace(uint32 Id);
 	FNewTrace			CreateTrace();
 	FAsioReadable*		OpenTrace(uint32 Id);
 
diff --git a/Engine/Source/Developer/TraceAnalysis/Private/Store/AsioStoreCborServer.cpp b/Engine/Source/Developer/TraceAnalysis/Private/Store/AsioStoreCborServer.cpp
--- a/Engine/Source/Developer/TraceAnalysis/Private/Store/AsioStoreCborServer.cpp
+++ b/Engine/Source/Developer/TraceAnalysis/Private/Store/AsioStoreCborServer.cpp
@@ -79,10 +79,7 @@ bool FAsioStoreCborPeer::IsOpen() const
 ////////////////////////////////////////////////////////////////////////////////
 void FAsioStoreCborPeer::Close()
 {
-	if (IsOpen())
-	{
-		Socket.Close();
-	}
+	Socket.Close();
 }
 
 ////////////////////////////////////////////////////////////////////////////////
@@ -353,13 +350,26 @@ FAsioStoreCborServer::FAsioStoreCborServer(
 ////////////////////////////////////////////////////////////////////////////////
 FAsioStoreCborServer::~FAsioStoreCborServer()
 {
+	Close();
+
 	for (FAsioStoreCborPeer* Peer : Peers)
 	{
-		Peer->Close();
 		delete Peer;
 	}
 }
 
+////////////////////////////////////////////////////////////////////////////////
+void FAsioStoreCborServer::Close()
+{
+	FAsioTcpServer::Close();
+	FAsioTickable::StopTick();
+
+	for (FAsioStoreCborPeer* Peer : Peers)
+	{
+		Peer->Close();
+	}
+}
+
 ////////////////////////////////////////////////////////////////////////////////
 FAsioStore& FAsioStoreCborServer::GetStore() const
 {
diff --git a/Engine/Source/Developer/TraceAnalysis/Private/Store/AsioStoreCborServer.h b/Engine/Source/Developer/TraceAnalysis/Private/Store/AsioStoreCborServer.h
--- a/Engine/Source/Developer/TraceAnalysis/Private/Store/AsioStoreCborServer.h
+++ b/Engine/Source/Developer/TraceAnalysis/Private/Store/AsioStoreCborServer.h
@@ -22,6 +22,7 @@ class FAsioStoreCborServer
 public:
 								FAsioStoreCborServer(asio::io_context& IoContext, FAsioStore& InStore, FAsioRecorder& InRecorder);
 								~FAsioStoreCborServer();
+	void						Close();
 	FAsioStore&					GetStore() const;
 	FAsioRecorder&				GetRecorder() const;
 
diff --git a/Engine/Source/Developer/TraceAnalysis/Private/Store/AsioTcpServer.cpp b/Engine/Source/Developer/TraceAnalysis/Private/Store/AsioTcpServer.cpp
--- a/Engine/Source/Developer/TraceAnalysis/Private/Store/AsioTcpServer.cpp
+++ b/Engine/Source/Developer/TraceAnalysis/Private/Store/AsioTcpServer.cpp
@@ -16,6 +16,16 @@ FAsioTcpServer::FAsioTcpServer(asio::io_context& IoContext)
 ////////////////////////////////////////////////////////////////////////////////
 FAsioTcpServer::~FAsioTcpServer()
 {
+	Close();
+}
+
+////////////////////////////////////////////////////////////////////////////////
+void FAsioTcpServer::Close()
+{
+	if (Acceptor.is_open())
+	{
+		Acceptor.close();
+	}
 }
 
 ////////////////////////////////////////////////////////////////////////////////
diff --git a/Engine/Source/Developer/TraceAnalysis/Private/Store/AsioTcpServer.h b/Engine/Source/Developer/TraceAnalysis/Private/Store/AsioTcpServer.h
--- a/Engine/Source/Developer/TraceAnalysis/Private/Store/AsioTcpServer.h
+++ b/Engine/Source/Developer/TraceAnalysis/Private/Store/AsioTcpServer.h
@@ -15,6 +15,7 @@ class FAsioTcpServer
 public:
 							FAsioTcpServer(asio::io_context& IoContext);
 	virtual					~FAsioTcpServer();
+	void					Close();
 	virtual bool			OnAccept(asio::ip::tcp::socket& Socket) = 0;
 	uint32					GetPort() const;
 	bool					StartServer(uint32 Port=0);
diff --git a/Engine/Source/Developer/TraceAnalysis/Private/Store/StoreService.cpp b/Engine/Source/Developer/TraceAnalysis/Private/Store/StoreService.cpp
--- a/Engine/Source/Developer/TraceAnalysis/Private/Store/StoreService.cpp
+++ b/Engine/Source/Developer/TraceAnalysis/Private/Store/StoreService.cpp
@@ -35,7 +35,10 @@ FStoreServiceImpl::FStoreServiceImpl(const FStoreService::FDesc& Desc)
 ////////////////////////////////////////////////////////////////////////////////
 FStoreServiceImpl::~FStoreServiceImpl()
 {
-	Context.Stop();
+	CborServer.Close();
+	Recorder.Close();
+	Store.Close();
+	Context.Wait();
 }
 
 ////////////////////////////////////////////////////////////////////////////////
