Close the numerous store service modules on the correct thread.

#rnx

diff --git a/Engine/Source/Developer/TraceAnalysis/Private/Store/StoreService.cpp b/Engine/Source/Developer/TraceAnalysis/Private/Store/StoreService.cpp
--- a/Engine/Source/Developer/TraceAnalysis/Private/Store/StoreService.cpp
+++ b/Engine/Source/Developer/TraceAnalysis/Private/Store/StoreService.cpp
@@ -35,9 +35,11 @@ FStoreServiceImpl::FStoreServiceImpl(const FStoreService::FDesc& Desc)
 ////////////////////////////////////////////////////////////////////////////////
 FStoreServiceImpl::~FStoreServiceImpl()
 {
-	CborServer.Close();
-	Recorder.Close();
-	Store.Close();
+	asio::post([this] () {
+		CborServer.Close();
+		Recorder.Close();
+		Store.Close();
+	});
 	Context.Wait();
 }
 
