Set linger to 0 for all store service module's listening and accepted sockets so Close() sends a RST instead of a FIN/ACK. This will prevent a flood of TIME_WAIT sockets that correctly following the passive/active close handshake.

#rnx

diff --git a/Engine/Source/Developer/TraceAnalysis/Private/Store/AsioTcpServer.cpp b/Engine/Source/Developer/TraceAnalysis/Private/Store/AsioTcpServer.cpp
--- a/Engine/Source/Developer/TraceAnalysis/Private/Store/AsioTcpServer.cpp
+++ b/Engine/Source/Developer/TraceAnalysis/Private/Store/AsioTcpServer.cpp
@@ -22,10 +22,7 @@ FAsioTcpServer::~FAsioTcpServer()
 ////////////////////////////////////////////////////////////////////////////////
 void FAsioTcpServer::Close()
 {
-	if (Acceptor.is_open())
-	{
-		Acceptor.close();
-	}
+	Acceptor.close();
 }
 
 ////////////////////////////////////////////////////////////////////////////////
@@ -68,6 +65,9 @@ bool FAsioTcpServer::StartServer(uint32 Port)
 		return false;
 	}
 
+	asio::socket_base::linger Linger(true, 0);
+	TempAcceptor.set_option(Linger);
+
 	Acceptor = MoveTemp(TempAcceptor);
 	AsyncAccept();
 	return true;
@@ -92,19 +92,16 @@ void FAsioTcpServer::AsyncAccept()
 
 	Acceptor.async_accept([this] (
 		const asio::error_code& ErrorCode,
-		tcp::socket Socket)
+		tcp::socket&& Socket)
 	{
-		if (ErrorCode)
-		{
-			return;
-		}
-
-		if (!OnAccept(Socket))
+		if (!ErrorCode && OnAccept(Socket))
 		{
+			AsyncAccept();
 			return;
 		}
 
-		AsyncAccept();
+		Close();
+		return;
 	});
 }
 
