Did pass over socket closing.

#rnx

diff --git a/Engine/Source/Developer/TraceAnalysis/Private/Store/AsioSocket.cpp b/Engine/Source/Developer/TraceAnalysis/Private/Store/AsioSocket.cpp
--- a/Engine/Source/Developer/TraceAnalysis/Private/Store/AsioSocket.cpp
+++ b/Engine/Source/Developer/TraceAnalysis/Private/Store/AsioSocket.cpp
@@ -64,7 +64,6 @@ uint32 FAsioSocket::GetLocalPort() const
 ////////////////////////////////////////////////////////////////////////////////
 void FAsioSocket::Close()
 {
-	Socket.shutdown(asio::ip::tcp::socket::shutdown_both);
 	Socket.close();
 }
 
diff --git a/Engine/Source/Developer/TraceAnalysis/Private/Store/AsioStoreCborServer.cpp b/Engine/Source/Developer/TraceAnalysis/Private/Store/AsioStoreCborServer.cpp
--- a/Engine/Source/Developer/TraceAnalysis/Private/Store/AsioStoreCborServer.cpp
+++ b/Engine/Source/Developer/TraceAnalysis/Private/Store/AsioStoreCborServer.cpp
@@ -310,7 +310,7 @@ void FAsioStoreCborPeer::OnIoComplete(uint32 Id, int32 Size)
 {
 	if (Size < 0)
 	{
-		Socket.Close();
+		Close();
 		return;
 	}
 
diff --git a/Engine/Source/Developer/TraceAnalysis/Private/Store/StoreClient.cpp b/Engine/Source/Developer/TraceAnalysis/Private/Store/StoreClient.cpp
--- a/Engine/Source/Developer/TraceAnalysis/Private/Store/StoreClient.cpp
+++ b/Engine/Source/Developer/TraceAnalysis/Private/Store/StoreClient.cpp
@@ -38,12 +38,6 @@ FTraceDataStream::~FTraceDataStream()
 ////////////////////////////////////////////////////////////////////////////////
 void FTraceDataStream::Close()
 {
-	if (!Socket.is_open())
-	{
-		return;
-	}
-
-	Socket.shutdown(asio::ip::tcp::socket::shutdown_both);
 	Socket.close();
 }
 
@@ -54,7 +48,7 @@ int32 FTraceDataStream::Read(void* Dest, uint32 DestSize)
 	size_t BytesRead = Socket.read_some(asio::buffer(Dest, DestSize), ErrorCode);
 	if (ErrorCode)
 	{
-		Socket.close();
+		Close();
 		return -1;
 	}
 
@@ -70,6 +64,7 @@ public:
 							FStoreCborClient();
 							~FStoreCborClient();
 	bool					IsOpen() const;
+	void					Close();
 	uint32					GetStoreAddress() const;
 	uint32					GetStorePort() const;
 	const FResponse&		GetResponse() const;
@@ -100,7 +95,12 @@ FStoreCborClient::FStoreCborClient()
 ////////////////////////////////////////////////////////////////////////////////
 FStoreCborClient::~FStoreCborClient()
 {
-	Socket.shutdown(asio::ip::tcp::socket::shutdown_both);
+	Close();
+}
+
+////////////////////////////////////////////////////////////////////////////////
+void FStoreCborClient::Close()
+{
 	Socket.close();
 }
 
@@ -273,12 +273,6 @@ FTraceDataStream* FStoreCborClient::ReadTrace(uint32 Id)
 		return nullptr;
 	}
 
-	SenderSocket.native_non_blocking(false, ErrorCode);
-	if (ErrorCode)
-	{
-		return nullptr;
-	}
-
 	return new FTraceDataStream(SenderSocket);
 }
 
