Prevent listening sockets from being inherited by child processes. See 11152126 for details.

#rnx

diff --git a/Engine/Source/Developer/TraceAnalysis/Private/Store/AsioTcpServer.cpp b/Engine/Source/Developer/TraceAnalysis/Private/Store/AsioTcpServer.cpp
--- a/Engine/Source/Developer/TraceAnalysis/Private/Store/AsioTcpServer.cpp
+++ b/Engine/Source/Developer/TraceAnalysis/Private/Store/AsioTcpServer.cpp
@@ -3,6 +3,11 @@
 #include "AsioTcpServer.h"
 #include "Templates/UnrealTemplate.h"
 
+#if PLATFORM_WINDOWS
+#	include <winsock2.h>
+#	pragma comment(lib, "ws2_32.lib")
+#endif
+
 namespace Trace
 {
 
@@ -46,7 +51,7 @@ uint32 FAsioTcpServer::GetPort() const
 }
 
 ////////////////////////////////////////////////////////////////////////////////
-bool FAsioTcpServer::StartServer(uint32 Port)
+bool FAsioTcpServer::StartServer(uint32 Port, uint32 Backlog)
 {
 	if (Acceptor.is_open())
 	{
@@ -55,12 +60,22 @@ bool FAsioTcpServer::StartServer(uint32 Port)
 
 	using asio::ip::tcp;
 
+	tcp::acceptor TempAcceptor(GetIoContext());
+
+#if PLATFORM_WINDOWS
+	DWORD Flags = WSA_FLAG_NO_HANDLE_INHERIT|WSA_FLAG_OVERLAPPED;
+	SOCKET Socket = WSASocket(AF_INET, SOCK_STREAM, IPPROTO_TCP, nullptr, 0, Flags);
+	TempAcceptor.assign(tcp::v4(), Socket);
+#else
+	TempAcceptor.open(tcp::v4());
+#endif
+
 	tcp::endpoint Endpoint(tcp::v4(), uint16(Port));
-	tcp::acceptor TempAcceptor(GetIoContext(), Endpoint, false);
+	TempAcceptor.bind(Endpoint);
+	TempAcceptor.listen(Backlog);
 
-	asio::error_code ErrorCode;
-	tcp::endpoint LocalEndpoint = TempAcceptor.local_endpoint(ErrorCode);
-	if (ErrorCode || LocalEndpoint.port() == 0)
+	tcp::endpoint LocalEndpoint = TempAcceptor.local_endpoint();
+	if (LocalEndpoint.port() == 0)
 	{
 		return false;
 	}
@@ -101,7 +116,6 @@ void FAsioTcpServer::AsyncAccept()
 		}
 
 		Close();
-		return;
 	});
 }
 
diff --git a/Engine/Source/Developer/TraceAnalysis/Private/Store/AsioTcpServer.h b/Engine/Source/Developer/TraceAnalysis/Private/Store/AsioTcpServer.h
--- a/Engine/Source/Developer/TraceAnalysis/Private/Store/AsioTcpServer.h
+++ b/Engine/Source/Developer/TraceAnalysis/Private/Store/AsioTcpServer.h
@@ -18,7 +18,7 @@ public:
 	void					Close();
 	virtual bool			OnAccept(asio::ip::tcp::socket& Socket) = 0;
 	uint32					GetPort() const;
-	bool					StartServer(uint32 Port=0);
+	bool					StartServer(uint32 Port=0, uint32 Backlog=1);
 	bool					StopServer();
 
 private:
