Removed hack to make live tracing update correctly. No longer relevant as we are using a different store-recorder backend.

#rnx

diff --git a/Engine/Source/Developer/TraceAnalysis/Private/DataStream.cpp b/Engine/Source/Developer/TraceAnalysis/Private/DataStream.cpp
--- a/Engine/Source/Developer/TraceAnalysis/Private/DataStream.cpp
+++ b/Engine/Source/Developer/TraceAnalysis/Private/DataStream.cpp
@@ -2,7 +2,6 @@
 
 #include "DataStream.h"
 #include "GenericPlatform/GenericPlatformFile.h"
-#include "HAL/PlatformProcess.h"
 
 namespace Trace
 {
@@ -24,18 +23,10 @@ int32 FFileStream::Read(void* Data, uint32 Size)
 	{
 		return 0;
 	}
-
-	uint64 Remaining;
-	while (true)
+	uint64 Remaining = End - Cursor;
+	if (Remaining == 0)
 	{
-		Remaining = End - Cursor;
-		if (Remaining)
-		{
-			break;
-		}
-
-		FPlatformProcess::Sleep(0.25f);
-		UpdateFileSize();
+		return 0;
 	}
 
 	uint64 Size64 = Size;
