Allow TraceAnalysis' Processor objects to be move-assigned as well as move-contructed.

#rnx

diff --git a/Engine/Source/Developer/TraceAnalysis/Private/Analysis/Processor.cpp b/Engine/Source/Developer/TraceAnalysis/Private/Analysis/Processor.cpp
--- a/Engine/Source/Developer/TraceAnalysis/Private/Analysis/Processor.cpp
+++ b/Engine/Source/Developer/TraceAnalysis/Private/Analysis/Processor.cpp
@@ -105,8 +105,15 @@ void FAnalysisProcessor::Pause(bool bState) { if (Impl != nullptr) { Impl->Pause
 
 ////////////////////////////////////////////////////////////////////////////////
 FAnalysisProcessor::FAnalysisProcessor(FAnalysisProcessor&& Rhs)
+{
+	this->operator = (MoveTemp(Rhs));
+}
+
+////////////////////////////////////////////////////////////////////////////////
+FAnalysisProcessor& FAnalysisProcessor::operator = (FAnalysisProcessor&& Rhs)
 {
 	Swap(Impl, Rhs.Impl);
+	return *this;
 }
 
 ////////////////////////////////////////////////////////////////////////////////
diff --git a/Engine/Source/Developer/TraceAnalysis/Public/Trace/Analysis.h b/Engine/Source/Developer/TraceAnalysis/Public/Trace/Analysis.h
--- a/Engine/Source/Developer/TraceAnalysis/Public/Trace/Analysis.h
+++ b/Engine/Source/Developer/TraceAnalysis/Public/Trace/Analysis.h
@@ -32,15 +32,15 @@ public:
 	 * @param bState Pause if true, resume if false. */
 	void Pause(bool bState);
 
-	~FAnalysisProcessor();
-	FAnalysisProcessor(FAnalysisProcessor&& Rhs);
+						~FAnalysisProcessor();
+						FAnalysisProcessor() = default;
+						FAnalysisProcessor(FAnalysisProcessor&& Rhs);
+	FAnalysisProcessor&	operator = (FAnalysisProcessor&&);
 
 private:
 	friend				class FAnalysisContext;
-						FAnalysisProcessor() = default;
 						FAnalysisProcessor(FAnalysisProcessor&) = delete;
-	FAnalysisProcessor&	operator = (FAnalysisProcessor&) = delete;
-	FAnalysisProcessor&	operator = (FAnalysisProcessor&&) = delete;
+	FAnalysisProcessor&	operator = (const FAnalysisProcessor&) = delete;
 	class				FImpl;
 	FImpl*				Impl = nullptr;
 };
