Explicitly end analysis on exit from the inner loop. Doing it implicitly in the destructor opens up the analysis framework for order of destruction problems.

#rnx

diff --git a/Engine/Source/Developer/TraceAnalysis/Private/Analysis/Engine.cpp b/Engine/Source/Developer/TraceAnalysis/Private/Analysis/Engine.cpp
--- a/Engine/Source/Developer/TraceAnalysis/Private/Analysis/Engine.cpp
+++ b/Engine/Source/Developer/TraceAnalysis/Private/Analysis/Engine.cpp
@@ -471,6 +471,11 @@ FAnalysisEngine::FAnalysisEngine(TArray<IAnalyzer*>&& InAnalyzers)
 
 ////////////////////////////////////////////////////////////////////////////////
 FAnalysisEngine::~FAnalysisEngine()
+{
+}
+
+////////////////////////////////////////////////////////////////////////////////
+void FAnalysisEngine::End()
 {
 	for (IAnalyzer* Analyzer : Analyzers)
 	{
diff --git a/Engine/Source/Developer/TraceAnalysis/Private/Analysis/Engine.h b/Engine/Source/Developer/TraceAnalysis/Private/Analysis/Engine.h
--- a/Engine/Source/Developer/TraceAnalysis/Private/Analysis/Engine.h
+++ b/Engine/Source/Developer/TraceAnalysis/Private/Analysis/Engine.h
@@ -22,6 +22,7 @@ public:
 						FAnalysisEngine(TArray<IAnalyzer*>&& InAnalyzers);
 						~FAnalysisEngine();
 	bool				OnData(FStreamReader& Reader);
+	void				End();
 
 private:
 	typedef bool (FAnalysisEngine::*ProtocolHandlerType)();
diff --git a/Engine/Source/Developer/TraceAnalysis/Private/Analysis/Processor.cpp b/Engine/Source/Developer/TraceAnalysis/Private/Analysis/Processor.cpp
--- a/Engine/Source/Developer/TraceAnalysis/Private/Analysis/Processor.cpp
+++ b/Engine/Source/Developer/TraceAnalysis/Private/Analysis/Processor.cpp
@@ -56,6 +56,7 @@ uint32 FAnalysisProcessor::FImpl::Run()
 		}
 	}
 
+	AnalysisEngine.End();
 	return 0;
 }
 
