Always attempt to cancel a timer regardless of the owning FAsioTickable's state. Without this it isn't possible to cancel one-shot ticks.

#rnx

diff --git a/Engine/Source/Developer/TraceAnalysis/Private/Store/AsioTickable.cpp b/Engine/Source/Developer/TraceAnalysis/Private/Store/AsioTickable.cpp
--- a/Engine/Source/Developer/TraceAnalysis/Private/Store/AsioTickable.cpp
+++ b/Engine/Source/Developer/TraceAnalysis/Private/Store/AsioTickable.cpp
@@ -53,11 +53,6 @@ bool FAsioTickable::StartTick(uint32 InMillisecondRate)
 ////////////////////////////////////////////////////////////////////////////////
 bool FAsioTickable::StopTick()
 {
-	if (!MillisecondRate)
-	{
-		return false;
-	}
-
 	MillisecondRate = 0;
 	Timer.cancel();
 	return true;
