Cancel the timeout tick when accepting a connection to read trace data from a relay. Otherwise the timeout fires causing an extra read, the loss of a page of trace data, and eventually a segfault or assert.

#rnx

diff --git a/Engine/Source/Developer/TraceAnalysis/Private/Store/AsioTraceRelay.cpp b/Engine/Source/Developer/TraceAnalysis/Private/Store/AsioTraceRelay.cpp
--- a/Engine/Source/Developer/TraceAnalysis/Private/Store/AsioTraceRelay.cpp
+++ b/Engine/Source/Developer/TraceAnalysis/Private/Store/AsioTraceRelay.cpp
@@ -70,6 +70,7 @@ bool FAsioTraceRelay::IsOpen() const
 bool FAsioTraceRelay::OnAccept(asio::ip::tcp::socket& Socket)
 {
 	Output = new FAsioSocket(Socket);
+	StopTick();
 	OnIoComplete(OpStart, 0);
 	return false;
 }
