Made the sockets that the trace store client creates uninheritable

diff --git a/Engine/Source/Developer/TraceAnalysis/Private/Store/StoreClient.cpp b/Engine/Source/Developer/TraceAnalysis/Private/Store/StoreClient.cpp
--- a/Engine/Source/Developer/TraceAnalysis/Private/Store/StoreClient.cpp
+++ b/Engine/Source/Developer/TraceAnalysis/Private/Store/StoreClient.cpp
@@ -160,6 +160,13 @@ bool FStoreCborClient::Connect(const TCHAR* Host, uint16 Port)
 	asio::ip::tcp::resolver::results_type Endpoints = Resolver.resolve(HostUtf8.Get(), PortString);
 
 	asio::error_code ErrorCode;
+#if PLATFORM_WINDOWS
+	DWORD Flags = WSA_FLAG_NO_HANDLE_INHERIT|WSA_FLAG_OVERLAPPED;
+	SOCKET WinSocket = WSASocket(AF_INET, SOCK_STREAM, IPPROTO_TCP, nullptr, 0, Flags);
+	Socket.assign(asio::ip::tcp::v4(), WinSocket);
+#else
+	Socket.open(asio::ip::tcp::v4());
+#endif
 
 	asio::connect(Socket, Endpoints, ErrorCode);
 	if (ErrorCode)
@@ -280,6 +287,14 @@ FTraceDataStream* FStoreCborClient::ReadTrace(uint32 Id)
 
 	asio::error_code ErrorCode;
 	asio::ip::tcp::socket SenderSocket(IoContext);
+#if PLATFORM_WINDOWS
+	DWORD Flags = WSA_FLAG_NO_HANDLE_INHERIT|WSA_FLAG_OVERLAPPED;
+	SOCKET WinSocket = WSASocket(AF_INET, SOCK_STREAM, IPPROTO_TCP, nullptr, 0, Flags);
+	SenderSocket.assign(asio::ip::tcp::v4(), WinSocket);
+#else
+	SenderSocket.open(asio::ip::tcp::v4());
+#endif
+
 	SenderSocket.connect(Endpoint, ErrorCode);
 	if (ErrorCode)
 	{
