Tweaked growth size of buffers that hold trace data during analysis.

#rnx

diff --git a/Engine/Source/Developer/TraceAnalysis/Private/Analysis/StreamReader.cpp b/Engine/Source/Developer/TraceAnalysis/Private/Analysis/StreamReader.cpp
--- a/Engine/Source/Developer/TraceAnalysis/Private/Analysis/StreamReader.cpp
+++ b/Engine/Source/Developer/TraceAnalysis/Private/Analysis/StreamReader.cpp
@@ -81,7 +81,7 @@ void FStreamBuffer::Consolidate()
 
 	if (DemandHint >= BufferSize)
 	{
-		const uint32 GrowthSizeMask = (8 << 10) - 1;
+		const uint32 GrowthSizeMask = (64 << 10) - 1;
 		BufferSize = (DemandHint + GrowthSizeMask + 1) & ~GrowthSizeMask;
 		Buffer = (uint8*)FMemory::Realloc(Buffer, BufferSize);
 	}
