Trace data is a stream and communication is one way. It is implemented this way to share code between sending trace data over the wire and writing it to a file. Because there's no ping/pong we can end up with a half-open TCP connection if the other end doesn't close its socket. So we'll enable keep-alive on the socket and set a short timeout (default is 2hrs).

diff --git a/Engine/Source/Developer/TraceAnalysis/Private/Store/AsioRecorder.cpp b/Engine/Source/Developer/TraceAnalysis/Private/Store/AsioRecorder.cpp
--- a/Engine/Source/Developer/TraceAnalysis/Private/Store/AsioRecorder.cpp
+++ b/Engine/Source/Developer/TraceAnalysis/Private/Store/AsioRecorder.cpp
@@ -5,6 +5,12 @@
 #include "AsioSocket.h"
 #include "AsioStore.h"
 
+#if PLATFORM_WINDOWS
+#	include "Windows/AllowWindowsPlatformTypes.h"
+#		include <Mstcpip.h>
+#	include "Windows/HideWindowsPlatformTypes.h"
+#endif
+
 namespace Trace
 {
 
@@ -166,6 +172,31 @@ bool FAsioRecorder::OnAccept(asio::ip::tcp::socket& Socket)
 		return true;
 	}
 
+#if PLATFORM_WINDOWS
+	// Trace data is a stream and communication is one way. It is implemented
+	// this way to share code between sending trace data over the wire and writing
+	// it to a file. Because there's no ping/pong we can end up with a half-open
+	// TCP connection if the other end doesn't close its socket. So we'll enable
+	// keep-alive on the socket and set a short timeout (default is 2hrs).
+	tcp_keepalive KeepAlive =
+	{
+		1,		// on
+		15000,	// timeout_ms
+		2000,	// interval_ms
+	};
+
+	DWORD BytesReturned;
+	WSAIoctl(
+		Socket.native_handle(),
+		SIO_KEEPALIVE_VALS,
+		&KeepAlive, sizeof(KeepAlive),
+		nullptr, 0,
+		&BytesReturned,
+		nullptr,
+		nullptr
+	);
+#endif
+
 	auto* Relay = new FAsioRecorderRelay(Socket, Trace.Writeable);
 
 	uint32 IdPieces[] = {
