Fixed MSVC static analysis warnings from a Windows SDK header.

#rnx

diff --git a/Engine/Source/Developer/TraceAnalysis/Private/Store/AsioRecorder.cpp b/Engine/Source/Developer/TraceAnalysis/Private/Store/AsioRecorder.cpp
--- a/Engine/Source/Developer/TraceAnalysis/Private/Store/AsioRecorder.cpp
+++ b/Engine/Source/Developer/TraceAnalysis/Private/Store/AsioRecorder.cpp
@@ -6,9 +6,11 @@
 #include "AsioStore.h"
 
 #if PLATFORM_WINDOWS
+	THIRD_PARTY_INCLUDES_START
 #	include "Windows/AllowWindowsPlatformTypes.h"
 #		include <Mstcpip.h>
 #	include "Windows/HideWindowsPlatformTypes.h"
+	THIRD_PARTY_INCLUDES_END
 #endif
 
 namespace Trace
