Simple provision for prefixing metadata at the start of a trace stream

#rnx

diff --git a/Engine/Source/Developer/TraceAnalysis/Private/Analysis/Engine.cpp b/Engine/Source/Developer/TraceAnalysis/Private/Analysis/Engine.cpp
--- a/Engine/Source/Developer/TraceAnalysis/Private/Analysis/Engine.cpp
+++ b/Engine/Source/Developer/TraceAnalysis/Private/Analysis/Engine.cpp
@@ -1218,10 +1218,12 @@ bool FAnalysisEngine::EstablishTransport(FStreamReader& Reader)
 
 	// Check for the magic uint32. Early traces did not include this as it was
 	// used to validate a inbound socket connection and then discarded.
-	if (Header->TransportVersion == 'E' || Header->TransportVersion == 'T')
+	if (Header->TransportVersion == 'E' ||
+		Header->TransportVersion == 'T' ||
+		Header->TransportVersion == '2')
 	{
 		const uint32* Magic = (const uint32*)(Reader.GetPointer(sizeof(*Magic)));
-		if (*Magic == 'ECRT')
+		if (*Magic == 'ECRT' || *Magic == '2CRT')
 		{
 			// Source is big-endian which we don't currently support
 			return false;
@@ -1233,6 +1235,16 @@ bool FAnalysisEngine::EstablishTransport(FStreamReader& Reader)
 			return EstablishTransport(Reader);
 		}
 
+		if (*Magic == 'TRC2')
+		{
+			Reader.Advance(sizeof(*Magic));
+			if (const uint16* MetadataSize = Reader.GetPointer<uint16>())
+			{
+				Reader.Advance(sizeof(*MetadataSize) + *MetadataSize);
+				return EstablishTransport(Reader);
+			}
+		}
+
 		return false;
 	}
 
diff --git a/Engine/Source/Runtime/TraceLog/Private/Trace/Writer.cpp b/Engine/Source/Runtime/TraceLog/Private/Trace/Writer.cpp
--- a/Engine/Source/Runtime/TraceLog/Private/Trace/Writer.cpp
+++ b/Engine/Source/Runtime/TraceLog/Private/Trace/Writer.cpp
@@ -371,8 +371,17 @@ static bool Writer_UpdateConnection()
 	GPendingDataHandle = 0;
 
 	// Handshake.
-	const uint32 Magic = 'TRCE';
-	bool bOk = IoWrite(GDataHandle, &Magic, sizeof(Magic));
+	struct FHandshake
+	{
+		uint32 Magic			= 'TRC2';
+		uint16 MetadataSize		= uint16(0);
+		enum
+		{
+			Size				= 6,
+		};
+	};
+	FHandshake Handshake;
+	bool bOk = IoWrite(GDataHandle, &Handshake, FHandshake::Size);
 
 	// Stream header
 	const struct {
