Expose active trace sessions' control port via the Trace Store.

#rnx

diff --git a/Engine/Source/Developer/TraceAnalysis/Private/Store/AsioRecorder.cpp b/Engine/Source/Developer/TraceAnalysis/Private/Store/AsioRecorder.cpp
--- a/Engine/Source/Developer/TraceAnalysis/Private/Store/AsioRecorder.cpp
+++ b/Engine/Source/Developer/TraceAnalysis/Private/Store/AsioRecorder.cpp
@@ -26,13 +26,17 @@ public:
 	bool				IsOpen();
 	void				Close();
 	uint32				GetIpAddress() const;
+	uint32				GetControlPort() const;
 
 private:
 	virtual void		OnIoComplete(uint32 Id, int32 Size) override;
+	bool				ReadMetadata(int32 Size);
 	static const uint32	BufferSize = 64 * 1024;
-	enum				{ OpStart, OpSocketRead, OpFileWrite };
+	enum				{ OpStart, OpSocketReadMetadata, OpSocketRead, OpFileWrite };
 	FAsioSocket			Input;
 	FAsioWriteable*		Output;
+	uint32				ActiveReadOp = OpSocketReadMetadata;
+	uint16				ControlPort = 0;
 	uint8				Buffer[BufferSize];
 };
 
@@ -71,6 +75,78 @@ uint32 FAsioRecorderRelay::GetIpAddress() const
 	return Input.GetRemoteAddress();
 }
 
+////////////////////////////////////////////////////////////////////////////////
+uint32 FAsioRecorderRelay::GetControlPort() const
+{
+	return ControlPort;
+}
+
+////////////////////////////////////////////////////////////////////////////////
+bool FAsioRecorderRelay::ReadMetadata(int32 Size)
+{
+	const uint8* Cursor = Buffer;
+	auto Read = [&] (int32 SizeToRead)
+	{
+		const uint8* Ptr = Cursor;
+		Cursor += SizeToRead;
+		Size -= SizeToRead;
+		return Ptr;
+	};
+
+	// Stream header
+	uint32 Magic;
+	if (Size < sizeof(Magic))
+	{
+		return true;
+	}
+
+	Magic = *(const uint32*)(Read(sizeof(Magic)));
+	switch (Magic)
+	{
+	case 'TRC2':		/* trace with metadata data */
+		break;
+
+	case 'TRCE':		/* valid, but to old or wrong endian for us */
+	case 'ECRT':
+	case '2CRT':
+		return true;
+
+	default:			/* unexpected magic */
+		return false;
+	}
+
+	// MetadataSize field
+	if (Size < 2)
+	{
+		return true;
+	}
+	Size = *(const uint16*)(Read(2));
+
+	// MetadataFields
+	while (Size >= 2)
+	{
+		struct {
+			uint8	Size;
+			uint8	Id;
+		} MetadataField;
+		MetadataField = *(const decltype(MetadataField)*)(Read(sizeof(MetadataField)));
+
+		if (Size < MetadataField.Size)
+		{
+			break;
+		}
+
+		if (MetadataField.Id == 0) /* ControlPortFieldId */
+		{
+			ControlPort = *(const uint16*)Cursor;
+		}
+
+		Size -= MetadataField.Size;
+	}
+
+	return true;
+}
+
 ////////////////////////////////////////////////////////////////////////////////
 void FAsioRecorderRelay::OnIoComplete(uint32 Id, int32 Size)
 {
@@ -82,13 +158,22 @@ void FAsioRecorderRelay::OnIoComplete(uint32 Id, int32 Size)
 
 	switch (Id)
 	{
+	case OpSocketReadMetadata:
+		ActiveReadOp = OpSocketRead;
+		if (!ReadMetadata(Size))
+		{
+			Close();
+			return;
+		}
+		/* fallthrough */
+
 	case OpSocketRead:
 		Output->Write(Buffer, Size, this, OpFileWrite);
 		break;
 
 	case OpStart:
 	case OpFileWrite:
-		Input.ReadSome(Buffer, BufferSize, this, OpSocketRead);
+		Input.ReadSome(Buffer, BufferSize, this, ActiveReadOp);
 		break;
 	}
 }
@@ -113,6 +198,12 @@ uint32 FAsioRecorder::FSession::GetIpAddress() const
 	return Relay->GetIpAddress();
 }
 
+////////////////////////////////////////////////////////////////////////////////
+uint32 FAsioRecorder::FSession::GetControlPort() const
+{
+	return Relay->GetControlPort();
+}
+
 
 
 ////////////////////////////////////////////////////////////////////////////////
diff --git a/Engine/Source/Developer/TraceAnalysis/Private/Store/AsioRecorder.h b/Engine/Source/Developer/TraceAnalysis/Private/Store/AsioRecorder.h
--- a/Engine/Source/Developer/TraceAnalysis/Private/Store/AsioRecorder.h
+++ b/Engine/Source/Developer/TraceAnalysis/Private/Store/AsioRecorder.h
@@ -25,6 +25,7 @@ public:
 		uint32					GetId() const;
 		uint32					GetTraceId() const;
 		uint32					GetIpAddress() const;
+		uint32					GetControlPort() const;
 
 	private:
 		friend					FAsioRecorder;
diff --git a/Engine/Source/Developer/TraceAnalysis/Private/Store/AsioStoreCborServer.cpp b/Engine/Source/Developer/TraceAnalysis/Private/Store/AsioStoreCborServer.cpp
--- a/Engine/Source/Developer/TraceAnalysis/Private/Store/AsioStoreCborServer.cpp
+++ b/Engine/Source/Developer/TraceAnalysis/Private/Store/AsioStoreCborServer.cpp
@@ -134,6 +134,7 @@ void FAsioStoreCborPeer::OnSessionInfo()
 	Builder.AddInteger("id", Session->GetId());
 	Builder.AddInteger("trace_id", Session->GetTraceId());
 	Builder.AddInteger("ip_address", Session->GetIpAddress());
+	Builder.AddInteger("control_port", Session->GetControlPort());
 	SendResponse(Builder.Done());
 }
 
diff --git a/Engine/Source/Developer/TraceAnalysis/Private/Store/StoreClient.cpp b/Engine/Source/Developer/TraceAnalysis/Private/Store/StoreClient.cpp
--- a/Engine/Source/Developer/TraceAnalysis/Private/Store/StoreClient.cpp
+++ b/Engine/Source/Developer/TraceAnalysis/Private/Store/StoreClient.cpp
@@ -416,6 +416,13 @@ uint32 FStoreClient::FSessionInfo::GetIpAddress() const
 	return Response->GetInteger("ip_address", 0);
 }
 
+////////////////////////////////////////////////////////////////////////////////
+uint32 FStoreClient::FSessionInfo::GetControlPort() const
+{
+	const auto* Response = (const FResponse*)this;
+	return Response->GetInteger("control_port", 0);
+}
+
 
 
 ////////////////////////////////////////////////////////////////////////////////
diff --git a/Engine/Source/Developer/TraceAnalysis/Public/Trace/StoreClient.h b/Engine/Source/Developer/TraceAnalysis/Public/Trace/StoreClient.h
--- a/Engine/Source/Developer/TraceAnalysis/Public/Trace/StoreClient.h
+++ b/Engine/Source/Developer/TraceAnalysis/Public/Trace/StoreClient.h
@@ -58,6 +58,7 @@ public:
 		uint32			GetId() const;
 		uint32			GetTraceId() const;
 		uint32			GetIpAddress() const;
+		uint32			GetControlPort() const;
 	};
 	uint32				GetSessionCount() const;
 	const FSessionInfo* GetSessionInfo(uint32 Index) const;
