Moved the Trace namespace into the UE namespace.

#rnx

diff --git a/Engine/Plugins/Animation/GameplayInsights/Source/GameplayInsights/Private/AnimationAnalyzer.h b/Engine/Plugins/Animation/GameplayInsights/Source/GameplayInsights/Private/AnimationAnalyzer.h
--- a/Engine/Plugins/Animation/GameplayInsights/Source/GameplayInsights/Private/AnimationAnalyzer.h
+++ b/Engine/Plugins/Animation/GameplayInsights/Source/GameplayInsights/Private/AnimationAnalyzer.h
@@ -7,7 +7,7 @@
 class FAnimationProvider;
 namespace TraceServices { class IAnalysisSession; }
 
-class FAnimationAnalyzer : public Trace::IAnalyzer
+class FAnimationAnalyzer : public UE::Trace::IAnalyzer
 {
 public:
 	FAnimationAnalyzer(TraceServices::IAnalysisSession& InSession, FAnimationProvider& InAnimationProvider);
diff --git a/Engine/Plugins/Animation/GameplayInsights/Source/GameplayInsights/Private/GameplayAnalyzer.h b/Engine/Plugins/Animation/GameplayInsights/Source/GameplayInsights/Private/GameplayAnalyzer.h
--- a/Engine/Plugins/Animation/GameplayInsights/Source/GameplayInsights/Private/GameplayAnalyzer.h
+++ b/Engine/Plugins/Animation/GameplayInsights/Source/GameplayInsights/Private/GameplayAnalyzer.h
@@ -7,7 +7,7 @@
 class FGameplayProvider;
 namespace TraceServices { class IAnalysisSession; }
 
-class FGameplayAnalyzer : public Trace::IAnalyzer
+class FGameplayAnalyzer : public UE::Trace::IAnalyzer
 {
 public:
 	FGameplayAnalyzer(TraceServices::IAnalysisSession& InSession, FGameplayProvider& InGameplayProvider);
diff --git a/Engine/Plugins/Animation/GameplayInsights/Source/GameplayInsights/Private/GameplayInsightsModule.cpp b/Engine/Plugins/Animation/GameplayInsights/Source/GameplayInsights/Private/GameplayInsightsModule.cpp
--- a/Engine/Plugins/Animation/GameplayInsights/Source/GameplayInsights/Private/GameplayInsightsModule.cpp
+++ b/Engine/Plugins/Animation/GameplayInsights/Source/GameplayInsights/Private/GameplayInsightsModule.cpp
@@ -154,11 +154,11 @@ void FGameplayInsightsModule::StartupModule()
 			{
 				// Create the Store Service.
 				FString StoreDir = FPaths::ProjectSavedDir() / TEXT("TraceSessions");
-				Trace::FStoreService::FDesc StoreServiceDesc;
+				UE::Trace::FStoreService::FDesc StoreServiceDesc;
 				StoreServiceDesc.StoreDir = *StoreDir;
 				StoreServiceDesc.RecorderPort = 0; // Let system decide port
 				StoreServiceDesc.ThreadCount = 2;
-				StoreService = TSharedPtr<Trace::FStoreService>(Trace::FStoreService::Create(StoreServiceDesc));
+				StoreService = TSharedPtr<UE::Trace::FStoreService>(UE::Trace::FStoreService::Create(StoreServiceDesc));
 
 				FCoreDelegates::OnPreExit.AddLambda([this]() {
 					StoreService.Reset();
@@ -166,7 +166,7 @@ void FGameplayInsightsModule::StartupModule()
 
 				// Connect to our newly created store and setup the insights module
 				UnrealInsightsModule.ConnectToStore(TEXT("localhost"), StoreService->GetPort());
-				Trace::SendTo(TEXT("localhost"), StoreService->GetRecorderPort());
+				UE::Trace::SendTo(TEXT("localhost"), StoreService->GetRecorderPort());
 
 				UnrealInsightsModule.CreateSessionViewer(false);
 				UnrealInsightsModule.StartAnalysisForLastLiveSession();
diff --git a/Engine/Plugins/Animation/GameplayInsights/Source/GameplayInsights/Private/GameplayInsightsModule.h b/Engine/Plugins/Animation/GameplayInsights/Source/GameplayInsights/Private/GameplayInsightsModule.h
--- a/Engine/Plugins/Animation/GameplayInsights/Source/GameplayInsights/Private/GameplayInsightsModule.h
+++ b/Engine/Plugins/Animation/GameplayInsights/Source/GameplayInsights/Private/GameplayInsightsModule.h
@@ -11,11 +11,14 @@
 
 struct FInsightsMajorTabExtender;
 
+namespace UE 
+{
 namespace Trace 
 {
 	class FStoreService;
 	class FStoreClient;
 }
+}
 
 class FGameplayInsightsModule : public IModuleInterface
 {
@@ -44,6 +47,6 @@ private:
 	FDelegateHandle CustomDebugObjectHandle;
 #endif
 
-	TSharedPtr<Trace::FStoreService> StoreService;
+	TSharedPtr<UE::Trace::FStoreService> StoreService;
 	TWeakPtr<FTabManager> WeakTimingProfilerTabManager;
 };
diff --git a/Engine/Plugins/Animation/GameplayInsights/Source/GameplayInsights/Private/ObjectPropertyTrace.cpp b/Engine/Plugins/Animation/GameplayInsights/Source/GameplayInsights/Private/ObjectPropertyTrace.cpp
--- a/Engine/Plugins/Animation/GameplayInsights/Source/GameplayInsights/Private/ObjectPropertyTrace.cpp
+++ b/Engine/Plugins/Animation/GameplayInsights/Source/GameplayInsights/Private/ObjectPropertyTrace.cpp
@@ -12,7 +12,7 @@ UE_TRACE_CHANNEL(ObjectProperties)
 
 UE_TRACE_EVENT_BEGIN(Object, ClassPropertyStringId, NoSync|Important)
 	UE_TRACE_EVENT_FIELD(uint32, Id)
-	UE_TRACE_EVENT_FIELD(Trace::WideString, Value)
+	UE_TRACE_EVENT_FIELD(UE::Trace::WideString, Value)
 UE_TRACE_EVENT_END()
 
 UE_TRACE_EVENT_BEGIN(Object, ClassProperty, NoSync|Important)
@@ -37,7 +37,7 @@ UE_TRACE_EVENT_BEGIN(Object, PropertyValue)
 	UE_TRACE_EVENT_FIELD(uint64, Cycle)
 	UE_TRACE_EVENT_FIELD(uint64, ObjectId)
 	UE_TRACE_EVENT_FIELD(int32, PropertyId)
-	UE_TRACE_EVENT_FIELD(Trace::WideString, Value)
+	UE_TRACE_EVENT_FIELD(UE::Trace::WideString, Value)
 UE_TRACE_EVENT_END()
 
 namespace ObjectPropertyTrace
diff --git a/Engine/Plugins/Developer/TraceDataFiltering/Source/TraceDataFiltering/Private/STraceDataFilterWidget.cpp b/Engine/Plugins/Developer/TraceDataFiltering/Source/TraceDataFiltering/Private/STraceDataFilterWidget.cpp
--- a/Engine/Plugins/Developer/TraceDataFiltering/Source/TraceDataFiltering/Private/STraceDataFilterWidget.cpp
+++ b/Engine/Plugins/Developer/TraceDataFiltering/Source/TraceDataFiltering/Private/STraceDataFilterWidget.cpp
@@ -736,14 +736,14 @@ void STraceDataFilterWidget::Tick(const FGeometry& AllottedGeometry, const doubl
 	{
 		if (AnalysisSession.IsValid())
 		{			
-			Trace::FStoreClient* StoreClient = InsightsModule.GetStoreClient();
+			UE::Trace::FStoreClient* StoreClient = InsightsModule.GetStoreClient();
 			if (StoreClient)
 			{
 				const int32 SessionCount = StoreClient->GetSessionCount();
 
 				if (SessionCount > 0)
 				{
-					const Trace::FStoreClient::FSessionInfo* SessionInfo = StoreClient->GetSessionInfo(SessionCount - 1);
+					const UE::Trace::FStoreClient::FSessionInfo* SessionInfo = StoreClient->GetSessionInfo(SessionCount - 1);
 					if (SessionInfo && (!AnalysisSession->IsAnalysisComplete() || SessionInfo->GetTraceId() != PreviousSessionHandle))
 					{
 						SetCurrentAnalysisSession(SessionInfo->GetTraceId(), AnalysisSession.ToSharedRef());
diff --git a/Engine/Plugins/Developer/TraceDataFiltering/Source/TraceDataFiltering/Private/SessionEditorFilterService.cpp b/Engine/Plugins/Developer/TraceDataFiltering/Source/TraceDataFiltering/Private/SessionEditorFilterService.cpp
--- a/Engine/Plugins/Developer/TraceDataFiltering/Source/TraceDataFiltering/Private/SessionEditorFilterService.cpp
+++ b/Engine/Plugins/Developer/TraceDataFiltering/Source/TraceDataFiltering/Private/SessionEditorFilterService.cpp
@@ -16,7 +16,7 @@ void FSessionEditorFilterService::OnApplyChannelChanges()
 	{
 		for (const FString& ChannelName : FrameEnabledChannels)
 		{
-			Trace::ToggleChannel(*ChannelName, true);
+			UE::Trace::ToggleChannel(*ChannelName, true);
 		}
 		FrameEnabledChannels.Empty();
 	}
@@ -25,7 +25,7 @@ void FSessionEditorFilterService::OnApplyChannelChanges()
 	{
 		for (const FString& ChannelName : FrameDisabledChannels)
 		{
-			Trace::ToggleChannel(*ChannelName, false);
+			UE::Trace::ToggleChannel(*ChannelName, false);
 		}
 
 		FrameDisabledChannels.Empty();
diff --git a/Engine/Plugins/Developer/TraceDataFiltering/Source/TraceDataFiltering/Private/Trace/SessionTraceFilterService.cpp b/Engine/Plugins/Developer/TraceDataFiltering/Source/TraceDataFiltering/Private/Trace/SessionTraceFilterService.cpp
--- a/Engine/Plugins/Developer/TraceDataFiltering/Source/TraceDataFiltering/Private/Trace/SessionTraceFilterService.cpp
+++ b/Engine/Plugins/Developer/TraceDataFiltering/Source/TraceDataFiltering/Private/Trace/SessionTraceFilterService.cpp
@@ -36,14 +36,14 @@ void FSessionTraceFilterService::OnApplyChannelChanges()
 	}
 
 	IUnrealInsightsModule& InsightsModule = FModuleManager::LoadModuleChecked<IUnrealInsightsModule>("TraceInsights");
-	Trace::FStoreClient* StoreClient = InsightsModule.GetStoreClient();
+	UE::Trace::FStoreClient* StoreClient = InsightsModule.GetStoreClient();
 
 	if (!StoreClient)
 	{
 		return;
 	}
 
-	const Trace::FStoreClient::FSessionInfo* SessionInfo = StoreClient->GetSessionInfoByTraceId(Handle);
+	const UE::Trace::FStoreClient::FSessionInfo* SessionInfo = StoreClient->GetSessionInfoByTraceId(Handle);
 
 	if (!SessionInfo)
 	{
@@ -55,7 +55,7 @@ void FSessionTraceFilterService::OnApplyChannelChanges()
 	ClientAddr->SetIp(SessionInfo->GetIpAddress());
 	ClientAddr->SetPort(1985);
 
-	Trace::FControlClient ControlClient;
+	UE::Trace::FControlClient ControlClient;
 	if (!ControlClient.Connect(ClientAddr.Get()))
 	{
 		return;
diff --git a/Engine/Plugins/Developer/TraceSourceFiltering/Source/SourceFilteringEditor/Private/STraceSourceFilteringWidget.cpp b/Engine/Plugins/Developer/TraceSourceFiltering/Source/SourceFilteringEditor/Private/STraceSourceFilteringWidget.cpp
--- a/Engine/Plugins/Developer/TraceSourceFiltering/Source/SourceFilteringEditor/Private/STraceSourceFilteringWidget.cpp
+++ b/Engine/Plugins/Developer/TraceSourceFiltering/Source/SourceFilteringEditor/Private/STraceSourceFilteringWidget.cpp
@@ -183,12 +183,12 @@ void STraceSourceFilteringWidget::Tick(const FGeometry& AllottedGeometry, const
 		TSharedPtr<const TraceServices::IAnalysisSession> AnalysisSession = InsightsModule.GetAnalysisSession();
 		if (AnalysisSession.IsValid())
 		{
-			Trace::FStoreClient* StoreClient = InsightsModule.GetStoreClient();
+			UE::Trace::FStoreClient* StoreClient = InsightsModule.GetStoreClient();
 			const int32 SessionCount = StoreClient->GetSessionCount();
 
 			if (SessionCount > 0)
 			{
-				const Trace::FStoreClient::FSessionInfo* SessionInfo = StoreClient->GetSessionInfo(SessionCount - 1);
+				const UE::Trace::FStoreClient::FSessionInfo* SessionInfo = StoreClient->GetSessionInfo(SessionCount - 1);
 				if (SessionInfo)
 				{
 					SetCurrentAnalysisSession(SessionInfo->GetTraceId(), AnalysisSession.ToSharedRef());
diff --git a/Engine/Plugins/Developer/TraceSourceFiltering/Source/SourceFilteringTrace/Private/SourceFilteringTraceModule.cpp b/Engine/Plugins/Developer/TraceSourceFiltering/Source/SourceFilteringTrace/Private/SourceFilteringTraceModule.cpp
--- a/Engine/Plugins/Developer/TraceSourceFiltering/Source/SourceFilteringTrace/Private/SourceFilteringTraceModule.cpp
+++ b/Engine/Plugins/Developer/TraceSourceFiltering/Source/SourceFilteringTrace/Private/SourceFilteringTraceModule.cpp
@@ -75,7 +75,7 @@ void FSourceFilteringTraceModule::StartupModule()
 	FTraceSourceFiltering::Initialize();
 
 	// Forcefully enable the source trace channel
-	Trace::ToggleChannel(TEXT("TraceSourceFiltersChannel"), true);
+	UE::Trace::ToggleChannel(TEXT("TraceSourceFiltersChannel"), true);
 
 #if WITH_EDITOR
 	// Add callback to trace out Filter Classes once the Asset Registry has finished loading 
diff --git a/Engine/Plugins/Experimental/ProxyLODPlugin/Source/ProxyLOD/Private/ProxyLODPlugin.cpp b/Engine/Plugins/Experimental/ProxyLODPlugin/Source/ProxyLOD/Private/ProxyLODPlugin.cpp
--- a/Engine/Plugins/Experimental/ProxyLODPlugin/Source/ProxyLOD/Private/ProxyLODPlugin.cpp
+++ b/Engine/Plugins/Experimental/ProxyLODPlugin/Source/ProxyLOD/Private/ProxyLODPlugin.cpp
@@ -286,7 +286,7 @@ public:
 		ThreadName = FString::Printf(TEXT("TBB %d"), ThreadIndex++);
 		ThreadID = FPlatformTLS::GetCurrentThreadId();
 		SetTls();
-		Trace::ThreadRegister(TEXT("TBB"), ThreadID, TPri_Normal);
+		UE::Trace::ThreadRegister(TEXT("TBB"), ThreadID, TPri_Normal);
 
 		FPlatformProcess::SetThreadName(*ThreadName);
 	}
diff --git a/Engine/Plugins/Runtime/NetworkPrediction/Source/NetworkPrediction/Private/NetworkPredictionModule.cpp b/Engine/Plugins/Runtime/NetworkPrediction/Source/NetworkPrediction/Private/NetworkPredictionModule.cpp
--- a/Engine/Plugins/Runtime/NetworkPrediction/Source/NetworkPrediction/Private/NetworkPredictionModule.cpp
+++ b/Engine/Plugins/Runtime/NetworkPrediction/Source/NetworkPrediction/Private/NetworkPredictionModule.cpp
@@ -39,14 +39,14 @@ class FNetworkPredictionModule : public INetworkPredictionModule
 void FNetworkPredictionModule::StartupModule()
 {
 	// Disable by default unless in the command line args. This is temp as the existing insights -trace parsing happen before the plugin is loaded
-	Trace::ToggleChannel(TEXT("NetworkPredictionChannel"), false);
+	UE::Trace::ToggleChannel(TEXT("NetworkPredictionChannel"), false);
 
 	FString EnabledChannels;
 	FParse::Value(FCommandLine::Get(), TEXT("-trace="), EnabledChannels, false);
 	UE::String::ParseTokens(EnabledChannels, TEXT(","), [](FStringView Token) {
 		if (Token.Compare(TEXT("NetworkPrediction"), ESearchCase::IgnoreCase)==0 || Token.Compare(TEXT("NP"), ESearchCase::IgnoreCase)==0)
 		{
-			Trace::ToggleChannel(TEXT("NetworkPredictionChannel"), true);
+		UE::Trace::ToggleChannel(TEXT("NetworkPredictionChannel"), true);
 		}
 	});
 
diff --git a/Engine/Plugins/Runtime/NetworkPrediction/Source/NetworkPrediction/Private/NetworkPredictionTrace.cpp b/Engine/Plugins/Runtime/NetworkPrediction/Source/NetworkPrediction/Private/NetworkPredictionTrace.cpp
--- a/Engine/Plugins/Runtime/NetworkPrediction/Source/NetworkPrediction/Private/NetworkPredictionTrace.cpp
+++ b/Engine/Plugins/Runtime/NetworkPrediction/Source/NetworkPrediction/Private/NetworkPredictionTrace.cpp
@@ -17,7 +17,7 @@
 #include "NetworkPredictionLog.h"
 #include "Trace/Trace.inl"
 
-// Also should update string tracing with Trace::AnsiString
+// Also should update string tracing with UE::Trace::AnsiString
 
 UE_TRACE_CHANNEL_DEFINE(NetworkPredictionChannel)
 
diff --git a/Engine/Plugins/Runtime/NetworkPredictionInsights/Source/NetworkPredictionInsights/Private/NetworkPredictionAnalyzer.h b/Engine/Plugins/Runtime/NetworkPredictionInsights/Source/NetworkPredictionInsights/Private/NetworkPredictionAnalyzer.h
--- a/Engine/Plugins/Runtime/NetworkPredictionInsights/Source/NetworkPredictionInsights/Private/NetworkPredictionAnalyzer.h
+++ b/Engine/Plugins/Runtime/NetworkPredictionInsights/Source/NetworkPredictionInsights/Private/NetworkPredictionAnalyzer.h
@@ -15,7 +15,7 @@ namespace TraceServices { class IAnalysisSession; }
 //
 // The analyzer is what populates the data in the Provider.
 
-class FNetworkPredictionAnalyzer : public Trace::IAnalyzer
+class FNetworkPredictionAnalyzer : public UE::Trace::IAnalyzer
 {
 public:
 
diff --git a/Engine/Plugins/Runtime/NetworkPredictionInsights/Source/NetworkPredictionInsights/Private/NetworkPredictionInsightsModule.cpp b/Engine/Plugins/Runtime/NetworkPredictionInsights/Source/NetworkPredictionInsights/Private/NetworkPredictionInsightsModule.cpp
--- a/Engine/Plugins/Runtime/NetworkPredictionInsights/Source/NetworkPredictionInsights/Private/NetworkPredictionInsightsModule.cpp
+++ b/Engine/Plugins/Runtime/NetworkPredictionInsights/Source/NetworkPredictionInsights/Private/NetworkPredictionInsightsModule.cpp
@@ -102,11 +102,11 @@ void FNetworkPredictionInsightsModule::StartupModule()
 			{
 				// Create the Store Service.
 				FString StoreDir = FPaths::ProjectSavedDir() / TEXT("TraceSessions");
-				Trace::FStoreService::FDesc StoreServiceDesc;
+				UE::Trace::FStoreService::FDesc StoreServiceDesc;
 				StoreServiceDesc.StoreDir = *StoreDir;
 				StoreServiceDesc.RecorderPort = 0; // Let system decide port
 				StoreServiceDesc.ThreadCount = 2;
-				StoreService = TSharedPtr<Trace::FStoreService>(Trace::FStoreService::Create(StoreServiceDesc));
+				StoreService = TSharedPtr<UE::Trace::FStoreService>(UE::Trace::FStoreService::Create(StoreServiceDesc));
 
 				FCoreDelegates::OnPreExit.AddLambda([this]() {
 					StoreService.Reset();
@@ -114,7 +114,7 @@ void FNetworkPredictionInsightsModule::StartupModule()
 
 				// Connect to our newly created store and setup the insights module
 				ensure(UnrealInsightsModule.ConnectToStore(TEXT("localhost"), StoreService->GetPort()));
-				Trace::SendTo(TEXT("localhost"), StoreService->GetRecorderPort());
+				UE::Trace::SendTo(TEXT("localhost"), StoreService->GetRecorderPort());
 
 				UnrealInsightsModule.CreateSessionViewer(false);
 				UnrealInsightsModule.StartAnalysisForLastLiveSession();
diff --git a/Engine/Plugins/Runtime/NetworkPredictionInsights/Source/NetworkPredictionInsights/Private/NetworkPredictionInsightsModule.h b/Engine/Plugins/Runtime/NetworkPredictionInsights/Source/NetworkPredictionInsights/Private/NetworkPredictionInsightsModule.h
--- a/Engine/Plugins/Runtime/NetworkPredictionInsights/Source/NetworkPredictionInsights/Private/NetworkPredictionInsightsModule.h
+++ b/Engine/Plugins/Runtime/NetworkPredictionInsights/Source/NetworkPredictionInsights/Private/NetworkPredictionInsightsModule.h
@@ -10,11 +10,14 @@
 
 struct FInsightsMajorTabExtender;
 
+namespace UE
+{
 namespace Trace 
 {
 	class FStoreService;
 	class FStoreClient;
 }
+}
 
 class FNetworkPredictionInsightsModule : public IModuleInterface
 {
@@ -29,7 +32,7 @@ private:
 	FDelegateHandle TickerHandle;
 	FDelegateHandle StoreServiceHandle;
 
-	TSharedPtr<Trace::FStoreService> StoreService;
+	TSharedPtr<UE::Trace::FStoreService> StoreService;
 	//TWeakPtr<FTabManager> WeakTimingProfilerTabManager;
 
 	static const FName InsightsTabName;	
diff --git a/Engine/Plugins/Slate/SlateInsights/Source/SlateInsights/Private/SlateAnalyzer.cpp b/Engine/Plugins/Slate/SlateInsights/Source/SlateInsights/Private/SlateAnalyzer.cpp
--- a/Engine/Plugins/Slate/SlateInsights/Source/SlateInsights/Private/SlateAnalyzer.cpp
+++ b/Engine/Plugins/Slate/SlateInsights/Source/SlateInsights/Private/SlateAnalyzer.cpp
@@ -21,7 +21,7 @@ FSlateAnalyzer::FSlateAnalyzer(TraceServices::IAnalysisSession& InSession, FSlat
 
 void FSlateAnalyzer::OnAnalysisBegin(const FOnAnalysisContext& Context)
 {
-	Trace::IAnalyzer::FInterfaceBuilder& Builder = Context.InterfaceBuilder;
+	UE::Trace::IAnalyzer::FInterfaceBuilder& Builder = Context.InterfaceBuilder;
 
 	Builder.RouteEvent(RouteId_ApplicationTickAndDrawWidgets, "SlateTrace", "ApplicationTickAndDrawWidgets");
 	Builder.RouteEvent(RouteId_AddWidget, "SlateTrace", "AddWidget");
@@ -37,7 +37,7 @@ bool FSlateAnalyzer::OnEvent(uint16 RouteId, EStyle Style, const FOnEventContext
 {
 	TraceServices::FAnalysisSessionEditScope _(Session);
 
-	const Trace::IAnalyzer::FEventData& EventData = Context.EventData;
+	const UE::Trace::IAnalyzer::FEventData& EventData = Context.EventData;
 	switch (RouteId)
 	{
 	case RouteId_ApplicationTickAndDrawWidgets:
diff --git a/Engine/Plugins/Slate/SlateInsights/Source/SlateInsights/Private/SlateAnalyzer.h b/Engine/Plugins/Slate/SlateInsights/Source/SlateInsights/Private/SlateAnalyzer.h
--- a/Engine/Plugins/Slate/SlateInsights/Source/SlateInsights/Private/SlateAnalyzer.h
+++ b/Engine/Plugins/Slate/SlateInsights/Source/SlateInsights/Private/SlateAnalyzer.h
@@ -13,7 +13,7 @@ namespace SlateInsights
 
 class FSlateProvider;
 
-class FSlateAnalyzer : public Trace::IAnalyzer
+class FSlateAnalyzer : public UE::Trace::IAnalyzer
 {
 public:
 	FSlateAnalyzer(TraceServices::IAnalysisSession& InSession, FSlateProvider& InSlateProvider);
diff --git a/Engine/Plugins/Slate/SlateInsights/Source/SlateInsights/Private/SlateProvider.cpp b/Engine/Plugins/Slate/SlateInsights/Source/SlateInsights/Private/SlateProvider.cpp
--- a/Engine/Plugins/Slate/SlateInsights/Source/SlateInsights/Private/SlateProvider.cpp
+++ b/Engine/Plugins/Slate/SlateInsights/Source/SlateInsights/Private/SlateProvider.cpp
@@ -16,7 +16,7 @@ FName FSlateProvider::ProviderName("SlateProvider");
 
 namespace Message
 {
-	FApplicationTickedMessage::FApplicationTickedMessage(const Trace::IAnalyzer::FEventData& EventData)
+	FApplicationTickedMessage::FApplicationTickedMessage(const UE::Trace::IAnalyzer::FEventData& EventData)
 		: DeltaTime(EventData.GetValue<float>("DeltaTime"))
 		, WidgetCount(EventData.GetValue<uint32>("WidgetCount"))
 		, TickCount(EventData.GetValue<uint32>("TickCount"))
@@ -31,7 +31,7 @@ namespace Message
 		static_assert(sizeof(ESlateTraceApplicationFlags) == sizeof(uint8), "ESlateTraceApplicationFlags is not a uint8");
 	}
 
-	FWidgetInfo::FWidgetInfo(const Trace::IAnalyzer::FEventData& EventData)
+	FWidgetInfo::FWidgetInfo(const UE::Trace::IAnalyzer::FEventData& EventData)
 		: WidgetId(EventData.GetValue<uint64>("WidgetId"))
 		, Path()
 		, DebugInfo()
@@ -41,14 +41,14 @@ namespace Message
 		EventData.GetString("DebugInfo", DebugInfo);
 	}
 	
-	FWidgetUpdatedMessage::FWidgetUpdatedMessage(const Trace::IAnalyzer::FEventData& EventData)
+	FWidgetUpdatedMessage::FWidgetUpdatedMessage(const UE::Trace::IAnalyzer::FEventData& EventData)
 		: WidgetId(EventData.GetValue<uint64>("WidgetId"))
 		, UpdateFlags(static_cast<EWidgetUpdateFlags>(EventData.GetValue<uint8>("UpdateFlags")))
 	{
 		static_assert(sizeof(EWidgetUpdateFlags) == sizeof(uint8), "EWidgetUpdateFlags is not a uint8");
 	}
 	
-	FWidgetInvalidatedMessage FWidgetInvalidatedMessage::FromWidget(const Trace::IAnalyzer::FEventData& EventData)
+	FWidgetInvalidatedMessage FWidgetInvalidatedMessage::FromWidget(const UE::Trace::IAnalyzer::FEventData& EventData)
 	{
 		static_assert(sizeof(EInvalidateWidgetReason) == sizeof(uint8), "EInvalidateWidgetReason is not a uint8");
 
@@ -61,7 +61,7 @@ namespace Message
 		return Message;
 	}
 
-	FWidgetInvalidatedMessage FWidgetInvalidatedMessage::FromRoot(const Trace::IAnalyzer::FEventData& EventData)
+	FWidgetInvalidatedMessage FWidgetInvalidatedMessage::FromRoot(const UE::Trace::IAnalyzer::FEventData& EventData)
 	{
 		FWidgetInvalidatedMessage Message;
 		Message.WidgetId = EventData.GetValue<uint64>("WidgetId");
@@ -72,7 +72,7 @@ namespace Message
 		return Message;
 	}
 
-	FWidgetInvalidatedMessage FWidgetInvalidatedMessage::FromChildOrder(const Trace::IAnalyzer::FEventData& EventData)
+	FWidgetInvalidatedMessage FWidgetInvalidatedMessage::FromChildOrder(const UE::Trace::IAnalyzer::FEventData& EventData)
 	{
 		FWidgetInvalidatedMessage Message;
 		Message.WidgetId = EventData.GetValue<uint64>("WidgetId");
@@ -83,7 +83,7 @@ namespace Message
 		return Message;
 	}
 
-	FString FWidgetInvalidatedMessage::GetCallstack(const Trace::IAnalyzer::FEventData& EventData)
+	FString FWidgetInvalidatedMessage::GetCallstack(const UE::Trace::IAnalyzer::FEventData& EventData)
 	{
 		FString Callstack = "";
 		const auto& CallstackReader = EventData.GetArrayView<uint64>("Callstack");
diff --git a/Engine/Plugins/Slate/SlateInsights/Source/SlateInsights/Private/SlateProvider.h b/Engine/Plugins/Slate/SlateInsights/Source/SlateInsights/Private/SlateProvider.h
--- a/Engine/Plugins/Slate/SlateInsights/Source/SlateInsights/Private/SlateProvider.h
+++ b/Engine/Plugins/Slate/SlateInsights/Source/SlateInsights/Private/SlateProvider.h
@@ -49,7 +49,7 @@ struct FWidgetInfo
 	uint64 EventIndex;
 
 	FWidgetInfo() = default;
-	FWidgetInfo(const Trace::IAnalyzer::FEventData& EventData);
+	FWidgetInfo(const UE::Trace::IAnalyzer::FEventData& EventData);
 	friend bool operator==(const FWidgetInfo& A, FWidgetId B) { return A.WidgetId == B; }
 };
 
@@ -59,7 +59,7 @@ struct FWidgetUpdatedMessage
 	/** Flag that was set by an invalidation or on the widget directly. */
 	EWidgetUpdateFlags UpdateFlags;
 
-	FWidgetUpdatedMessage(const Trace::IAnalyzer::FEventData& EventData);
+	FWidgetUpdatedMessage(const UE::Trace::IAnalyzer::FEventData& EventData);
 };
 
 struct FWidgetInvalidatedMessage
@@ -72,10 +72,10 @@ struct FWidgetInvalidatedMessage
 	FString ScriptTrace;
 	FString Callstack;
 
-	static FWidgetInvalidatedMessage FromWidget(const Trace::IAnalyzer::FEventData& EventData);
-	static FWidgetInvalidatedMessage FromRoot(const Trace::IAnalyzer::FEventData& EventData);
-	static FWidgetInvalidatedMessage FromChildOrder(const Trace::IAnalyzer::FEventData& EventData);
-	static FString GetCallstack(const Trace::IAnalyzer::FEventData& EventData);
+	static FWidgetInvalidatedMessage FromWidget(const UE::Trace::IAnalyzer::FEventData& EventData);
+	static FWidgetInvalidatedMessage FromRoot(const UE::Trace::IAnalyzer::FEventData& EventData);
+	static FWidgetInvalidatedMessage FromChildOrder(const UE::Trace::IAnalyzer::FEventData& EventData);
+	static FString GetCallstack(const UE::Trace::IAnalyzer::FEventData& EventData);
 };
 
 struct FApplicationTickedMessage
@@ -91,7 +91,7 @@ struct FApplicationTickedMessage
 	uint32 RootInvalidatedCount;
 	ESlateTraceApplicationFlags Flags;
 
-	FApplicationTickedMessage(const Trace::IAnalyzer::FEventData& EventData);
+	FApplicationTickedMessage(const UE::Trace::IAnalyzer::FEventData& EventData);
 };
 
 } //namespace Message
diff --git a/Engine/Source/Developer/TraceAnalysis/Private/Analysis/Context.cpp b/Engine/Source/Developer/TraceAnalysis/Private/Analysis/Context.cpp
--- a/Engine/Source/Developer/TraceAnalysis/Private/Analysis/Context.cpp
+++ b/Engine/Source/Developer/TraceAnalysis/Private/Analysis/Context.cpp
@@ -3,8 +3,8 @@
 #include "Processor.h"
 #include "Trace/Analysis.h"
 
-namespace Trace
-{
+namespace UE {
+namespace Trace {
 
 ////////////////////////////////////////////////////////////////////////////////
 void FAnalysisContext::AddAnalyzer(IAnalyzer& Analyzer)
@@ -24,3 +24,4 @@ FAnalysisProcessor FAnalysisContext::Process(IInDataStream& DataStream)
 }
 
 } // namespace Trace
+} // namespace UE
diff --git a/Engine/Source/Developer/TraceAnalysis/Private/Analysis/Engine.cpp b/Engine/Source/Developer/TraceAnalysis/Private/Analysis/Engine.cpp
--- a/Engine/Source/Developer/TraceAnalysis/Private/Analysis/Engine.cpp
+++ b/Engine/Source/Developer/TraceAnalysis/Private/Analysis/Engine.cpp
@@ -15,8 +15,8 @@
 #include "Transport/Transport.h"
 #include "Transport/TidPacketTransport.h"
 
-namespace Trace
-{
+namespace UE {
+namespace Trace {
 
 ////////////////////////////////////////////////////////////////////////////////
 void SerializeToCborImpl(TArray<uint8>&, const IAnalyzer::FEventData&, uint32);
@@ -442,7 +442,7 @@ IAnalyzer::FEventFieldInfo::EType IAnalyzer::FEventFieldInfo::GetType() const
 {
 	const auto* Inner = (const FAnalysisEngine::FDispatch::FField*)this;
 
-	if (Inner->Class == Protocol0::Field_String)
+	if (Inner->Class == UE::Trace::Protocol0::Field_String)
 	{
 		return (Inner->SizeAndType == 1) ? EType::AnsiString : EType::WideString;
 	}
@@ -581,7 +581,7 @@ bool IAnalyzer::FEventData::GetString(const ANSICHAR* FieldName, FAnsiStringView
 	}
 
 	const auto& Field = Dispatch.Fields[Index];
-	if (Field.Class != Protocol0::Field_String || Field.SizeAndType != sizeof(ANSICHAR))
+	if (Field.Class != UE::Trace::Protocol0::Field_String || Field.SizeAndType != sizeof(ANSICHAR))
 	{
 		return false;
 	}
@@ -608,7 +608,7 @@ bool IAnalyzer::FEventData::GetString(const ANSICHAR* FieldName, FStringView& Ou
 	}
 
 	const auto& Field = Dispatch.Fields[Index];
-	if (Field.Class != Protocol0::Field_String || Field.SizeAndType != sizeof(TCHAR))
+	if (Field.Class != UE::Trace::Protocol0::Field_String || Field.SizeAndType != sizeof(TCHAR))
 	{
 		return false;
 	}
@@ -641,7 +641,7 @@ bool IAnalyzer::FEventData::GetString(const ANSICHAR* FieldName, FString& Out) c
 	}
 
 	const auto& Field = Dispatch.Fields[Index];
-	if (Field.Class == Protocol0::Field_String)
+	if (Field.Class == UE::Trace::Protocol0::Field_String)
 	{
 		if (Field.SizeAndType == sizeof(ANSICHAR))
 		{
@@ -717,6 +717,8 @@ FAnalysisEngine::~FAnalysisEngine()
 ////////////////////////////////////////////////////////////////////////////////
 void FAnalysisEngine::Begin()
 {
+	using namespace UE::Trace;
+
 	// Call out to all registered analyzers to have them register event interest
 	struct : IAnalyzer::FInterfaceBuilder
 	{
@@ -937,7 +939,7 @@ void FAnalysisEngine::OnNewTrace(const FOnEventContext& Context)
 	Hint &= Serial.Mask;
 	Serial.Value = Hint|0x80000000;
 
-	UserUidBias = EventData.GetValue<uint32>("UserUidBias", uint32(Protocol3::EKnownEventUids::User));
+	UserUidBias = EventData.GetValue<uint32>("UserUidBias", uint32(UE::Trace::Protocol3::EKnownEventUids::User));
 
 	OnTiming(Context);
 }
@@ -1070,6 +1072,8 @@ void FAnalysisEngine::ForEachRoute(uint32 RouteIndex, bool bScoped, ImplType&& I
 ////////////////////////////////////////////////////////////////////////////////
 void FAnalysisEngine::OnNewEventInternal(const void* EventData)
 {
+	using namespace UE::Trace;
+
 	FDispatchBuilder Builder;
 	switch (ProtocolVersion)
 	{
@@ -1148,6 +1152,8 @@ bool FAnalysisEngine::AddDispatch(FDispatch* Dispatch)
 ////////////////////////////////////////////////////////////////////////////////
 void FAnalysisEngine::OnNewEventProtocol0(FDispatchBuilder& Builder, const void* EventData)
 {
+	using namespace UE::Trace;
+
 	const auto& NewEvent = *(Protocol0::FNewEventEvent*)(EventData);
 
 	const auto* NameCursor = (const ANSICHAR*)(NewEvent.Fields + NewEvent.FieldCount);
@@ -1184,6 +1190,8 @@ void FAnalysisEngine::OnNewEventProtocol0(FDispatchBuilder& Builder, const void*
 ////////////////////////////////////////////////////////////////////////////////
 void FAnalysisEngine::OnNewEventProtocol1(FDispatchBuilder& Builder, const void* EventData)
 {
+	using namespace UE::Trace;
+
 	OnNewEventProtocol0(Builder, EventData);
 
 	const auto& NewEvent = *(Protocol1::FNewEventEvent*)(EventData);
@@ -1207,6 +1215,8 @@ void FAnalysisEngine::OnNewEventProtocol1(FDispatchBuilder& Builder, const void*
 ////////////////////////////////////////////////////////////////////////////////
 bool FAnalysisEngine::EstablishTransport(FStreamReader& Reader)
 {
+	using namespace UE::Trace;
+
 	const struct {
 		uint8 TransportVersion;
 		uint8 ProtocolVersion;
@@ -1323,6 +1333,8 @@ bool FAnalysisEngine::OnData(FStreamReader& Reader)
 ////////////////////////////////////////////////////////////////////////////////
 bool FAnalysisEngine::OnDataProtocol0()
 {
+	using namespace UE::Trace;
+
 	FThreads::FInfo ThreadInfo;
 
 	while (true)
@@ -1417,7 +1429,7 @@ bool FAnalysisEngine::OnDataProtocol2()
 			FThreads::FInfo& ThreadInfo = Threads.GetInfo(RotaItem.ThreadId);
 
 			uint32 AvailableSerial;
-			if (ProtocolVersion == Protocol4::EProtocol::Id)
+			if (ProtocolVersion == UE::Trace::Protocol4::EProtocol::Id)
 			{
 				AvailableSerial = OnDataProtocol4(*(RotaItem.Reader), ThreadInfo);
 			}
@@ -1512,6 +1524,8 @@ bool FAnalysisEngine::OnDataProtocol2()
 ////////////////////////////////////////////////////////////////////////////////
 int32 FAnalysisEngine::OnDataProtocol2(FStreamReader& Reader, FThreads::FInfo& ThreadInfo)
 {
+	using namespace UE::Trace;
+
 	while (true)
 	{
 		auto Mark = Reader.SaveMark();
@@ -1624,6 +1638,8 @@ int32 FAnalysisEngine::OnDataProtocol2(FStreamReader& Reader, FThreads::FInfo& T
 ////////////////////////////////////////////////////////////////////////////////
 int32 FAnalysisEngine::OnDataProtocol2Aux(FStreamReader& Reader, FAuxDataCollector& Collector)
 {
+	using namespace UE::Trace;
+
 	while (true)
 	{
 		const uint8* NextByte = Reader.GetPointer<uint8>();
@@ -1683,7 +1699,7 @@ int32 FAnalysisEngine::OnDataProtocol4Known(
 	FStreamReader& Reader,
 	FThreads::FInfo& ThreadInfo)
 {
-	using namespace Protocol4;
+	using namespace UE::Trace::Protocol4;
 
 	switch (Uid)
 	{
@@ -1770,7 +1786,7 @@ int32 FAnalysisEngine::OnDataProtocol4Impl(
 {
 	/* Returns 0 if an event was successfully processed, 1 if there's not enough
 	 * data available, or ~AvailableLogSerial if the pending event is in the future */
-	using namespace Protocol4;
+	using namespace UE::Trace::Protocol4;
 
 	auto Mark = Reader.SaveMark();
 
@@ -1829,7 +1845,7 @@ int32 FAnalysisEngine::OnDataProtocol4Impl(
 	// Make sure we consume events in the correct order
 	if ((Dispatch->Flags & FDispatch::Flag_NoSync) == 0)
 	{
-		const auto* HeaderSync = (Protocol4::FEventHeaderSync*)Header;
+		const auto* HeaderSync = (UE::Trace::Protocol4::FEventHeaderSync*)Header;
 		uint32 EventSerial = HeaderSync->SerialLow|(uint32(HeaderSync->SerialHigh) << 16);
 		if (EventSerial != (Serial.Value & Serial.Mask))
 		{
@@ -1924,3 +1940,4 @@ int32 FAnalysisEngine::OnDataProtocol4Impl(
 }
 
 } // namespace Trace
+} // namespace UE
diff --git a/Engine/Source/Developer/TraceAnalysis/Private/Analysis/Engine.h b/Engine/Source/Developer/TraceAnalysis/Private/Analysis/Engine.h
--- a/Engine/Source/Developer/TraceAnalysis/Private/Analysis/Engine.h
+++ b/Engine/Source/Developer/TraceAnalysis/Private/Analysis/Engine.h
@@ -6,8 +6,8 @@
 #include "Trace/Analysis.h"
 #include "Trace/Analyzer.h"
 
-namespace Trace
-{
+namespace UE {
+namespace Trace {
 
 class FStreamReader;
 
@@ -127,3 +127,4 @@ private:
 };
 
 } // namespace Trace
+} // namespace UE
diff --git a/Engine/Source/Developer/TraceAnalysis/Private/Analysis/Transport/PacketTransport.cpp b/Engine/Source/Developer/TraceAnalysis/Private/Analysis/Transport/PacketTransport.cpp
--- a/Engine/Source/Developer/TraceAnalysis/Private/Analysis/Transport/PacketTransport.cpp
+++ b/Engine/Source/Developer/TraceAnalysis/Private/Analysis/Transport/PacketTransport.cpp
@@ -5,18 +5,21 @@
 
 #include <initializer_list>
 
-namespace Trace
-{
-
 ////////////////////////////////////////////////////////////////////////////////
-namespace Private
-{
+namespace UE {
+namespace Trace {
+namespace Private {
 
 TRACELOG_API int32 Decode(const void*, int32, void*, int32);
 
 } // namespace Private
+} // namespace Trace
+} // namespace UE
+
 
 
+namespace UE {
+namespace Trace {
 
 ////////////////////////////////////////////////////////////////////////////////
 struct FPacketTransport::FPacketNode
@@ -150,7 +153,7 @@ bool FPacketTransport::GetNextBatch()
 			};
 			auto* PacketEncoded = (FPacketEncoded*)PacketBase;
 
-			Node->Size = Private::Decode(
+			Node->Size = UE::Trace::Private::Decode(
 				PacketEncoded->Data,
 				int32(PacketEncoded->PacketSize - sizeof(FPacketEncoded)),
 				Node->Data,
@@ -175,3 +178,4 @@ bool FPacketTransport::GetNextBatch()
 }
 
 } // namespace Trace
+} // namespace UE
diff --git a/Engine/Source/Developer/TraceAnalysis/Private/Analysis/Transport/TidPacketTransport.cpp b/Engine/Source/Developer/TraceAnalysis/Private/Analysis/Transport/TidPacketTransport.cpp
--- a/Engine/Source/Developer/TraceAnalysis/Private/Analysis/Transport/TidPacketTransport.cpp
+++ b/Engine/Source/Developer/TraceAnalysis/Private/Analysis/Transport/TidPacketTransport.cpp
@@ -5,18 +5,21 @@
 #include "HAL/UnrealMemory.h"
 #include "Trace/Detail/Protocol.h"
 
-namespace Trace
-{
-
 ////////////////////////////////////////////////////////////////////////////////
-namespace Private
-{
+namespace UE {
+namespace Trace {
+namespace Private {
 
 TRACELOG_API int32 Decode(const void*, int32, void*, int32);
 
 } // namespace Private
+} // namespace Trace
+} // namespace UE
+
 
 
+namespace UE {
+namespace Trace {
 
 ////////////////////////////////////////////////////////////////////////////////
 bool FTidPacketTransport::ReadPacket()
@@ -49,7 +52,7 @@ bool FTidPacketTransport::ReadPacket()
 		uint16* DecodedSize = (uint16*)(PacketBase + 1);
 		uint8* Dest = Thread.Buffer.Append(*DecodedSize);
 		DataSize -= sizeof(*DecodedSize);
-		int32 ResultSize = Private::Decode(DecodedSize + 1, DataSize, Dest, *DecodedSize);
+		int32 ResultSize = UE::Trace::Private::Decode(DecodedSize + 1, DataSize, Dest, *DecodedSize);
 		check(int32(*DecodedSize) == ResultSize);
 	}
 	else
@@ -77,7 +80,7 @@ FTidPacketTransport::FThreadStream& FTidPacketTransport::FindOrAddThread(uint32
 
 	// Internal events are sent over tid 0 and they should be processed first. We
 	// don't care about the order otherwise.
-	if (Thread.ThreadId == ETransportTid::Internal)
+	if (Thread.ThreadId == UE::Trace::ETransportTid::Internal)
 	{
 		Threads.Insert(Thread, 0);
 		return Threads[0];
@@ -117,3 +120,4 @@ int32 FTidPacketTransport::GetThreadId(uint32 Index) const
 }
 
 } // namespace Trace
+} // namespace UE
diff --git a/Engine/Source/Developer/TraceAnalysis/Private/Store/StoreClient.cpp b/Engine/Source/Developer/TraceAnalysis/Private/Store/StoreClient.cpp
--- a/Engine/Source/Developer/TraceAnalysis/Private/Store/StoreClient.cpp
+++ b/Engine/Source/Developer/TraceAnalysis/Private/Store/StoreClient.cpp
@@ -6,8 +6,8 @@
 #include "CborPayload.h"
 #include "Templates/UnrealTemplate.h"
 
-namespace Trace
-{
+namespace UE {
+namespace Trace {
 
 ////////////////////////////////////////////////////////////////////////////////
 class FTraceDataStream
@@ -342,10 +342,11 @@ bool FStoreCborClient::GetSessionInfoByTraceId(uint32 TraceId)
 }
 
 } // namespace Trace
+} // namespace UE
 
 
-namespace Trace
-{
+namespace UE {
+namespace Trace {
 
 ////////////////////////////////////////////////////////////////////////////////
 uint32 FStoreClient::FStatus::GetRecorderPort() const
@@ -578,3 +579,4 @@ const FStoreClient::FSessionInfo* FStoreClient::GetSessionInfoByTraceId(uint32 T
 }
 
 } // namespace Trace
+} // namespace UE
diff --git a/Engine/Source/Developer/TraceAnalysis/Public/Trace/Analyzer.h b/Engine/Source/Developer/TraceAnalysis/Public/Trace/Analyzer.h
--- a/Engine/Source/Developer/TraceAnalysis/Public/Trace/Analyzer.h
+++ b/Engine/Source/Developer/TraceAnalysis/Public/Trace/Analyzer.h
@@ -5,11 +5,9 @@
 #include "CoreTypes.h"
 #include "Containers/ArrayView.h"
 #include "Containers/StringView.h"
-#include "Logging/LogMacros.h"
-#include "Trace/Detail/Field.h"
 
-namespace Trace
-{
+namespace UE {
+namespace Trace {
 
 /**
  * Interface that users implement to analyze the events in a trace. Analysis
@@ -326,3 +324,4 @@ const ValueType* IAnalyzer::TArrayReader<ValueType>::GetData() const
 }
 
 } // namespace Trace
+} // namespace UE
diff --git a/Engine/Source/Developer/TraceAnalysis/Public/Trace/DataStream.h b/Engine/Source/Developer/TraceAnalysis/Public/Trace/DataStream.h
--- a/Engine/Source/Developer/TraceAnalysis/Public/Trace/DataStream.h
+++ b/Engine/Source/Developer/TraceAnalysis/Public/Trace/DataStream.h
@@ -4,8 +4,8 @@
 
 #include "CoreTypes.h"
 
-namespace Trace
-{
+namespace UE {
+namespace Trace {
 
 class IInDataStream
 {
@@ -15,4 +15,5 @@ public:
 	virtual void	Close() {}
 };
 
-}
+} // namespace Trace
+} // namespace UE
diff --git a/Engine/Source/Developer/TraceInsights/Private/Insights/InsightsManager.cpp b/Engine/Source/Developer/TraceInsights/Private/Insights/InsightsManager.cpp
--- a/Engine/Source/Developer/TraceInsights/Private/Insights/InsightsManager.cpp
+++ b/Engine/Source/Developer/TraceInsights/Private/Insights/InsightsManager.cpp
@@ -313,7 +313,7 @@ void FInsightsManager::OnSessionInfoTabClosed(TSharedRef<SDockTab> TabBeingClose
 
 bool FInsightsManager::ConnectToStore(const TCHAR* Host, uint32 Port)
 {
-	using namespace Trace;
+	using namespace UE::Trace;
 	FStoreClient* Client = FStoreClient::Connect(Host, Port);
 	StoreClient = TUniquePtr<FStoreClient>(Client);
 	return StoreClient.IsValid();
@@ -507,7 +507,7 @@ void FInsightsManager::LoadLastLiveSession()
 		return;
 	}
 
-	const Trace::FStoreClient::FSessionInfo* SessionInfo = StoreClient->GetSessionInfo(SessionCount - 1);
+	const UE::Trace::FStoreClient::FSessionInfo* SessionInfo = StoreClient->GetSessionInfo(SessionCount - 1);
 	if (SessionInfo == nullptr)
 	{
 		return;
@@ -531,7 +531,7 @@ void FInsightsManager::LoadTrace(uint32 InTraceId, bool InAutoQuit)
 		return;
 	}
 
-	Trace::FStoreClient::FTraceData TraceData = StoreClient->ReadTrace(InTraceId);
+	UE::Trace::FStoreClient::FTraceData TraceData = StoreClient->ReadTrace(InTraceId);
 	if (!TraceData)
 	{
 		if (InAutoQuit)
@@ -542,7 +542,7 @@ void FInsightsManager::LoadTrace(uint32 InTraceId, bool InAutoQuit)
 	}
 
 	FString TraceName;
-	const Trace::FStoreClient::FTraceInfo* TraceInfo = StoreClient->GetTraceInfoById(InTraceId);
+	const UE::Trace::FStoreClient::FTraceInfo* TraceInfo = StoreClient->GetTraceInfoById(InTraceId);
 	if (TraceInfo != nullptr)
 	{
 		FAnsiStringView Name = TraceInfo->GetName();
diff --git a/Engine/Source/Developer/TraceInsights/Private/Insights/InsightsManager.h b/Engine/Source/Developer/TraceInsights/Private/Insights/InsightsManager.h
--- a/Engine/Source/Developer/TraceInsights/Private/Insights/InsightsManager.h
+++ b/Engine/Source/Developer/TraceInsights/Private/Insights/InsightsManager.h
@@ -16,10 +16,13 @@
 
 ////////////////////////////////////////////////////////////////////////////////////////////////////
 
+namespace UE
+{
 namespace Trace
 {
 	class FStoreClient;
 }
+}
 
 namespace TraceServices
 {
@@ -98,7 +101,7 @@ public:
 	const FString& GetStoreDir() const { return StoreDir; }
 
 	bool ConnectToStore(const TCHAR* Host, uint32 Port);
-	Trace::FStoreClient* GetStoreClient() const { return StoreClient.Get(); }
+	UE::Trace::FStoreClient* GetStoreClient() const { return StoreClient.Get(); }
 
 	/** @return an instance of the trace analysis session. */
 	TSharedPtr<const TraceServices::IAnalysisSession> GetSession() const;
@@ -274,7 +277,7 @@ private:
 	FString StoreDir;
 
 	/** The client used to connect to the trace store. */
-	TUniquePtr<Trace::FStoreClient> StoreClient;
+	TUniquePtr<UE::Trace::FStoreClient> StoreClient;
 
 	/** The trace analysis session. */
 	TSharedPtr<const TraceServices::IAnalysisSession> Session;
diff --git a/Engine/Source/Developer/TraceInsights/Private/Insights/StoreService/DiagnosticsSessionAnalyzer.h b/Engine/Source/Developer/TraceInsights/Private/Insights/StoreService/DiagnosticsSessionAnalyzer.h
--- a/Engine/Source/Developer/TraceInsights/Private/Insights/StoreService/DiagnosticsSessionAnalyzer.h
+++ b/Engine/Source/Developer/TraceInsights/Private/Insights/StoreService/DiagnosticsSessionAnalyzer.h
@@ -8,7 +8,7 @@
 namespace Insights
 {
 
-struct FDiagnosticsSessionAnalyzer : public Trace::IAnalyzer
+struct FDiagnosticsSessionAnalyzer : public UE::Trace::IAnalyzer
 {
 	virtual void OnAnalysisBegin(const FOnAnalysisContext& Context) override;
 	virtual bool OnEvent(uint16 RouteId, EStyle, const FOnEventContext& Context) override;
diff --git a/Engine/Source/Developer/TraceInsights/Private/Insights/StoreService/StoreBrowser.cpp b/Engine/Source/Developer/TraceInsights/Private/Insights/StoreService/StoreBrowser.cpp
--- a/Engine/Source/Developer/TraceInsights/Private/Insights/StoreService/StoreBrowser.cpp
+++ b/Engine/Source/Developer/TraceInsights/Private/Insights/StoreService/StoreBrowser.cpp
@@ -78,7 +78,7 @@ void FStoreBrowser::Exit()
 
 ////////////////////////////////////////////////////////////////////////////////////////////////////
 
-Trace::FStoreClient* FStoreBrowser::GetStoreClient() const
+UE::Trace::FStoreClient* FStoreBrowser::GetStoreClient() const
 {
 	// TODO: thread safety !?
 	return FInsightsManager::Get()->GetStoreClient();
@@ -88,7 +88,7 @@ Trace::FStoreClient* FStoreBrowser::GetStoreClient() const
 
 void FStoreBrowser::UpdateTraces()
 {
-	Trace::FStoreClient* StoreClient = GetStoreClient();
+	UE::Trace::FStoreClient* StoreClient = GetStoreClient();
 	if (StoreClient == nullptr)
 	{
 		ResetTraces();
@@ -100,7 +100,7 @@ void FStoreBrowser::UpdateTraces()
 
 	// Check if the list of trace sessions has changed.
 	{
-		const Trace::FStoreClient::FStatus* Status = StoreClient->GetStatus();
+		const UE::Trace::FStoreClient::FStatus* Status = StoreClient->GetStatus();
 		if (StoreChangeSerial != Status->GetChangeSerial())
 		{
 			StoreChangeSerial = Status->GetChangeSerial();
@@ -110,7 +110,7 @@ void FStoreBrowser::UpdateTraces()
 				for (int32 TraceIndex = Traces.Num() - 1; TraceIndex >= 0; --TraceIndex)
 				{
 					const uint32 TraceId = Traces[TraceIndex]->TraceId;
-					const Trace::FStoreClient::FTraceInfo* TraceInfo = StoreClient->GetTraceInfoById(TraceId);
+					const UE::Trace::FStoreClient::FTraceInfo* TraceInfo = StoreClient->GetTraceInfoById(TraceId);
 					if (TraceInfo == nullptr)
 					{
 						FScopeLock Lock(&TracesCriticalSection);
@@ -126,7 +126,7 @@ void FStoreBrowser::UpdateTraces()
 				const int32 TraceCount = StoreClient->GetTraceCount();
 				for (int32 TraceIndex = 0; TraceIndex < TraceCount; ++TraceIndex)
 				{
-					const Trace::FStoreClient::FTraceInfo* TraceInfo = StoreClient->GetTraceInfo(TraceIndex);
+					const UE::Trace::FStoreClient::FTraceInfo* TraceInfo = StoreClient->GetTraceInfo(TraceIndex);
 					if (TraceInfo != nullptr)
 					{
 						const uint32 TraceId = TraceInfo->GetId();
@@ -172,14 +172,14 @@ void FStoreBrowser::UpdateTraces()
 
 			FDateTime Timestamp(0);
 			uint64 Size = 0;
-			const Trace::FStoreClient::FTraceInfo* TraceInfo = StoreClient->GetTraceInfoById(TraceId);
+			const UE::Trace::FStoreClient::FTraceInfo* TraceInfo = StoreClient->GetTraceInfoById(TraceId);
 			if (TraceInfo != nullptr)
 			{
 				Timestamp = FStoreBrowserTraceInfo::ConvertTimestamp(TraceInfo->GetTimestamp());
 				Size = TraceInfo->GetSize();
 			}
 
-			const Trace::FStoreClient::FSessionInfo* SessionInfo = StoreClient->GetSessionInfoByTraceId(TraceId);
+			const UE::Trace::FStoreClient::FSessionInfo* SessionInfo = StoreClient->GetSessionInfoByTraceId(TraceId);
 			if (SessionInfo != nullptr)
 			{
 				// The trace is still live.
@@ -223,7 +223,7 @@ void FStoreBrowser::UpdateTraces()
 		const uint32 SessionCount = StoreClient->GetSessionCount();
 		for (uint32 SessionIndex = 0; SessionIndex < SessionCount; ++SessionIndex)
 		{
-			const Trace::FStoreClient::FSessionInfo* SessionInfo = StoreClient->GetSessionInfo(SessionIndex);
+			const UE::Trace::FStoreClient::FSessionInfo* SessionInfo = StoreClient->GetSessionInfo(SessionIndex);
 			if (SessionInfo != nullptr)
 			{
 				const uint32 TraceId = SessionInfo->GetTraceId();
@@ -239,7 +239,7 @@ void FStoreBrowser::UpdateTraces()
 
 						FDateTime Timestamp(0);
 						uint64 Size = 0;
-						const Trace::FStoreClient::FTraceInfo* TraceInfo = StoreClient->GetTraceInfoById(TraceId);
+						const UE::Trace::FStoreClient::FTraceInfo* TraceInfo = StoreClient->GetTraceInfoById(TraceId);
 						if (TraceInfo != nullptr)
 						{
 							Timestamp = FStoreBrowserTraceInfo::ConvertTimestamp(TraceInfo->GetTimestamp());
@@ -315,7 +315,7 @@ void FStoreBrowser::UpdateMetadata(FStoreBrowserTraceInfo& Trace)
 {
 	//UE_LOG(TraceInsights, Log, TEXT("[StoreBrowser] Updating metadata for trace 0x%08X (%s)..."), TraceSession.TraceId, *TraceSession.Name.ToString());
 
-	Trace::FStoreClient* StoreClient = GetStoreClient();
+	UE::Trace::FStoreClient* StoreClient = GetStoreClient();
 	if (StoreClient == nullptr)
 	{
 		return;
@@ -323,13 +323,13 @@ void FStoreBrowser::UpdateMetadata(FStoreBrowserTraceInfo& Trace)
 
 	FPlatformProcess::SleepNoStats(0.0f);
 
-	Trace::FStoreClient::FTraceData TraceData = StoreClient->ReadTrace(Trace.TraceId);
+	UE::Trace::FStoreClient::FTraceData TraceData = StoreClient->ReadTrace(Trace.TraceId);
 	if (!TraceData)
 	{
 		return;
 	}
 
-	struct FDataStream : public Trace::IInDataStream
+	struct FDataStream : public UE::Trace::IInDataStream
 	{
 		virtual int32 Read(void* Data, uint32 Size) override
 		{
@@ -350,14 +350,14 @@ void FStoreBrowser::UpdateMetadata(FStoreBrowserTraceInfo& Trace)
 		}
 
 		int32 BytesRead = 0;
-		Trace::IInDataStream* Inner;
+		UE::Trace::IInDataStream* Inner;
 	};
 
 	FDataStream DataStream;
 	DataStream.Inner = TraceData.Get();
 
 	FDiagnosticsSessionAnalyzer Analyzer;
-	Trace::FAnalysisContext Context;
+	UE::Trace::FAnalysisContext Context;
 	Context.AddAnalyzer(Analyzer);
 	Context.Process(DataStream).Wait();
 
diff --git a/Engine/Source/Developer/TraceInsights/Private/Insights/StoreService/StoreBrowser.h b/Engine/Source/Developer/TraceInsights/Private/Insights/StoreService/StoreBrowser.h
--- a/Engine/Source/Developer/TraceInsights/Private/Insights/StoreService/StoreBrowser.h
+++ b/Engine/Source/Developer/TraceInsights/Private/Insights/StoreService/StoreBrowser.h
@@ -10,10 +10,13 @@
 #include "HAL/CriticalSection.h"
 #include "Templates/Atomic.h"
 
+namespace UE
+{
 namespace Trace
 {
 	class FStoreClient;
 }
+}
 
 namespace Insights
 {
@@ -81,7 +84,7 @@ public:
 	void Unlock() { check(bTracesLocked); bTracesLocked = false; TracesCriticalSection.Unlock(); }
 
 private:
-	Trace::FStoreClient* GetStoreClient() const;
+	UE::Trace::FStoreClient* GetStoreClient() const;
 
 	void UpdateTraces();
 	void ResetTraces();
diff --git a/Engine/Source/Developer/TraceInsights/Private/Insights/TraceInsightsModule.cpp b/Engine/Source/Developer/TraceInsights/Private/Insights/TraceInsightsModule.cpp
--- a/Engine/Source/Developer/TraceInsights/Private/Insights/TraceInsightsModule.cpp
+++ b/Engine/Source/Developer/TraceInsights/Private/Insights/TraceInsightsModule.cpp
@@ -108,11 +108,11 @@ void FTraceInsightsModule::CreateDefaultStore()
 	FInsightsManager::Get()->SetStoreDir(StoreDir);
 
 	// Create the Store Service.
-	Trace::FStoreService::FDesc StoreServiceDesc;
+	UE::Trace::FStoreService::FDesc StoreServiceDesc;
 	StoreServiceDesc.StoreDir = *StoreDir;
 	StoreServiceDesc.RecorderPort = 1980;
 	StoreServiceDesc.ThreadCount = 2;
-	StoreService = TUniquePtr<Trace::FStoreService>(Trace::FStoreService::Create(StoreServiceDesc));
+	StoreService = TUniquePtr<UE::Trace::FStoreService>(UE::Trace::FStoreService::Create(StoreServiceDesc));
 
 	if (StoreService.IsValid())
 	{
@@ -122,7 +122,7 @@ void FTraceInsightsModule::CreateDefaultStore()
 
 ////////////////////////////////////////////////////////////////////////////////////////////////////
 
-Trace::FStoreClient* FTraceInsightsModule::GetStoreClient()
+UE::Trace::FStoreClient* FTraceInsightsModule::GetStoreClient()
 {
 	return FInsightsManager::Get()->GetStoreClient();
 }
diff --git a/Engine/Source/Developer/TraceInsights/Private/Insights/TraceInsightsModule.h b/Engine/Source/Developer/TraceInsights/Private/Insights/TraceInsightsModule.h
--- a/Engine/Source/Developer/TraceInsights/Private/Insights/TraceInsightsModule.h
+++ b/Engine/Source/Developer/TraceInsights/Private/Insights/TraceInsightsModule.h
@@ -6,10 +6,13 @@
 #include "Insights/IUnrealInsightsModule.h"
 #include "Framework/Docking/TabManager.h"
 
+namespace UE
+{
 namespace Trace
 {
 	class FStoreService;
 }
+}
 
 namespace TraceServices
 {
@@ -37,7 +40,7 @@ public:
 
 	virtual void CreateDefaultStore() override;
 
-	virtual Trace::FStoreClient* GetStoreClient() override;
+	virtual UE::Trace::FStoreClient* GetStoreClient() override;
 	virtual bool ConnectToStore(const TCHAR* InStoreHost, uint32 InStorePort) override;
 
 	virtual void CreateSessionBrowser(bool bAllowDebugTools = false, bool bSingleProcess = false, bool bInitializeTesting = false) override;
@@ -93,7 +96,7 @@ protected:
 	void OnWindowClosedEvent(const TSharedRef<SWindow>&);
 
 protected:
-	TUniquePtr<Trace::FStoreService> StoreService;
+	TUniquePtr<UE::Trace::FStoreService> StoreService;
 
 	TSharedPtr<TraceServices::IAnalysisService> TraceAnalysisService;
 	TSharedPtr<TraceServices::IModuleService> TraceModuleService;
diff --git a/Engine/Source/Developer/TraceInsights/Private/Insights/Widgets/SStartPageWindow.cpp b/Engine/Source/Developer/TraceInsights/Private/Insights/Widgets/SStartPageWindow.cpp
--- a/Engine/Source/Developer/TraceInsights/Private/Insights/Widgets/SStartPageWindow.cpp
+++ b/Engine/Source/Developer/TraceInsights/Private/Insights/Widgets/SStartPageWindow.cpp
@@ -1181,7 +1181,7 @@ FReply SStartPageWindow::Connect_OnClicked()
 	}
 
 	bool bConnectedSuccessfully = false;
-	Trace::FControlClient ControlClient;
+	UE::Trace::FControlClient ControlClient;
 	if (ControlClient.Connect(*HostText.ToString()))
 	{
 		TSharedPtr<FInternetAddr> RecorderAddr;
@@ -1230,7 +1230,7 @@ FReply SStartPageWindow::Connect_OnClicked()
 
 void SStartPageWindow::RefreshTraceList()
 {
-	Trace::FStoreClient* StoreClient = FInsightsManager::Get()->GetStoreClient();
+	UE::Trace::FStoreClient* StoreClient = FInsightsManager::Get()->GetStoreClient();
 	if (StoreClient == nullptr)
 	{
 		return;
@@ -1722,7 +1722,7 @@ TSharedRef<SWidget> SStartPageWindow::MakeTraceListMenu()
 
 	MenuBuilder.BeginSection("AvailableTraces", LOCTEXT("AvailableTracesHeading", "Top Most Recently Created Traces"));
 	{
-		Trace::FStoreClient* StoreClient = FInsightsManager::Get()->GetStoreClient();
+		UE::Trace::FStoreClient* StoreClient = FInsightsManager::Get()->GetStoreClient();
 		if (StoreClient != nullptr)
 		{
 			// Make a copy of the trace list (to allow list view to be sorted by other criteria).
@@ -1830,7 +1830,7 @@ FReply SStartPageWindow::ExploreTraceStoreDirectory_OnClicked()
 
 FText SStartPageWindow::GetRecorderStatusText() const
 {
-	Trace::FStoreClient* StoreClient = FInsightsManager::Get()->GetStoreClient();
+	UE::Trace::FStoreClient* StoreClient = FInsightsManager::Get()->GetStoreClient();
 	const bool bIsRecorderServerRunning = (StoreClient != nullptr); // TODO: StoreClient->IsRecorderServerRunning();
 
 	if (bIsRecorderServerRunning)
@@ -1847,7 +1847,7 @@ FText SStartPageWindow::GetRecorderStatusText() const
 
 EVisibility SStartPageWindow::StartTraceRecorder_Visibility() const
 {
-	Trace::FStoreClient* StoreClient = FInsightsManager::Get()->GetStoreClient();
+	UE::Trace::FStoreClient* StoreClient = FInsightsManager::Get()->GetStoreClient();
 	const bool bIsRecorderServerRunning = (StoreClient != nullptr); // TODO: StoreClient->IsRecorderServerRunning();
 
 	return bIsRecorderServerRunning ? EVisibility::Collapsed : EVisibility::Visible;
@@ -1857,7 +1857,7 @@ EVisibility SStartPageWindow::StartTraceRecorder_Visibility() const
 
 EVisibility SStartPageWindow::StopTraceRecorder_Visibility() const
 {
-	Trace::FStoreClient* StoreClient = FInsightsManager::Get()->GetStoreClient();
+	UE::Trace::FStoreClient* StoreClient = FInsightsManager::Get()->GetStoreClient();
 	const bool bIsRecorderServerRunning = (StoreClient != nullptr); // TODO: StoreClient->IsRecorderServerRunning();
 
 	return bIsRecorderServerRunning ? EVisibility::Visible : EVisibility::Collapsed;
diff --git a/Engine/Source/Developer/TraceInsights/Public/Insights/IUnrealInsightsModule.h b/Engine/Source/Developer/TraceInsights/Public/Insights/IUnrealInsightsModule.h
--- a/Engine/Source/Developer/TraceInsights/Public/Insights/IUnrealInsightsModule.h
+++ b/Engine/Source/Developer/TraceInsights/Public/Insights/IUnrealInsightsModule.h
@@ -10,10 +10,13 @@
 
 class FExtender;
 
+namespace UE
+{
 namespace Trace
 {
 	class FStoreClient;
 }
+}
 
 namespace TraceServices
 {
@@ -177,7 +180,7 @@ public:
 	/**
 	 * Gets the store client.
 	 */
-	virtual Trace::FStoreClient* GetStoreClient() = 0;
+	virtual UE::Trace::FStoreClient* GetStoreClient() = 0;
 
 	/**
 	 * Connects to a specified store.
diff --git a/Engine/Source/Developer/TraceServices/Private/AnalysisService.cpp b/Engine/Source/Developer/TraceServices/Private/AnalysisService.cpp
--- a/Engine/Source/Developer/TraceServices/Private/AnalysisService.cpp
+++ b/Engine/Source/Developer/TraceServices/Private/AnalysisService.cpp
@@ -79,7 +79,7 @@ void FAnalysisSessionLock::EndEdit()
 	}
 }
 
-FAnalysisSession::FAnalysisSession(const TCHAR* SessionName, TUniquePtr<Trace::IInDataStream>&& InDataStream)
+FAnalysisSession::FAnalysisSession(const TCHAR* SessionName, TUniquePtr<UE::Trace::IInDataStream>&& InDataStream)
 	: Name(SessionName)
 	, DurationSeconds(0.0)
 	, Allocator(32 << 20)
@@ -102,8 +102,8 @@ FAnalysisSession::~FAnalysisSession()
 
 void FAnalysisSession::Start()
 {
-	Trace::FAnalysisContext Context;
-	for (Trace::IAnalyzer* Analyzer : ReadAnalyzers())
+	UE::Trace::FAnalysisContext Context;
+	for (UE::Trace::IAnalyzer* Analyzer : ReadAnalyzers())
 	{
 		Context.AddAnalyzer(*Analyzer);
 	}
@@ -124,7 +124,7 @@ void FAnalysisSession::Wait() const
 	Processor.Wait();
 }
 
-void FAnalysisSession::AddAnalyzer(Trace::IAnalyzer* Analyzer)
+void FAnalysisSession::AddAnalyzer(UE::Trace::IAnalyzer* Analyzer)
 {
 	Analyzers.Add(Analyzer);
 }
@@ -180,7 +180,7 @@ TSharedPtr<const IAnalysisSession> FAnalysisService::Analyze(const TCHAR* Sessio
 TSharedPtr<const IAnalysisSession> FAnalysisService::StartAnalysis(const TCHAR* SessionUri)
 {
 	struct FFileDataStream
-		: public Trace::IInDataStream
+		: public UE::Trace::IInDataStream
 	{
 		virtual int32 Read(void* Data, uint32 Size) override
 		{
@@ -209,11 +209,11 @@ TSharedPtr<const IAnalysisSession> FAnalysisService::StartAnalysis(const TCHAR*
 	FileStream->Handle = TUniquePtr<IFileHandle>(Handle);
 	FileStream->Remaining = Handle->Size();
 
-	TUniquePtr<Trace::IInDataStream> DataStream(FileStream);
+	TUniquePtr<UE::Trace::IInDataStream> DataStream(FileStream);
 	return StartAnalysis(SessionUri, MoveTemp(DataStream));
 }
 
-TSharedPtr<const IAnalysisSession> FAnalysisService::StartAnalysis(const TCHAR* SessionName, TUniquePtr<Trace::IInDataStream>&& DataStream)
+TSharedPtr<const IAnalysisSession> FAnalysisService::StartAnalysis(const TCHAR* SessionName, TUniquePtr<UE::Trace::IInDataStream>&& DataStream)
 {
 	TSharedRef<FAnalysisSession> Session = MakeShared<FAnalysisSession>(SessionName, MoveTemp(DataStream));
 
diff --git a/Engine/Source/Developer/TraceServices/Private/AnalysisServicePrivate.h b/Engine/Source/Developer/TraceServices/Private/AnalysisServicePrivate.h
--- a/Engine/Source/Developer/TraceServices/Private/AnalysisServicePrivate.h
+++ b/Engine/Source/Developer/TraceServices/Private/AnalysisServicePrivate.h
@@ -34,7 +34,7 @@ class FAnalysisSession
 	: public IAnalysisSession
 {
 public:
-	FAnalysisSession(const TCHAR* SessionName, TUniquePtr<Trace::IInDataStream>&& InDataStream);
+	FAnalysisSession(const TCHAR* SessionName, TUniquePtr<UE::Trace::IInDataStream>&& InDataStream);
 	virtual ~FAnalysisSession();
 	void Start();
 	virtual void Stop(bool bAndWait) const override;
@@ -58,10 +58,10 @@ public:
 	virtual void ReadAccessCheck() const override { return Lock.ReadAccessCheck(); }
 	virtual void WriteAccessCheck() override { return Lock.WriteAccessCheck(); }
 
-	virtual void AddAnalyzer(Trace::IAnalyzer* Analyzer) override;
+	virtual void AddAnalyzer(UE::Trace::IAnalyzer* Analyzer) override;
 	virtual void AddProvider(const FName& Name, IProvider* Provider) override;
 
-	const TArray<Trace::IAnalyzer*> ReadAnalyzers() { return Analyzers; }
+	const TArray<UE::Trace::IAnalyzer*> ReadAnalyzers() { return Analyzers; }
 
 private:
 	virtual const IProvider* ReadProviderPrivate(const FName& Name) const override;
@@ -73,11 +73,11 @@ private:
 	double DurationSeconds = 0.0;
 	FSlabAllocator Allocator;
 	FStringStore StringStore;
-	TArray<Trace::IAnalyzer*> Analyzers;
+	TArray<UE::Trace::IAnalyzer*> Analyzers;
 	TArray<IProvider*> Providers;
 	TMap<FName, IProvider*> ProvidersMap;
-	mutable TUniquePtr<Trace::IInDataStream> DataStream;
-	mutable Trace::FAnalysisProcessor Processor;
+	mutable TUniquePtr<UE::Trace::IInDataStream> DataStream;
+	mutable UE::Trace::FAnalysisProcessor Processor;
 };
 
 class FAnalysisService
@@ -88,7 +88,7 @@ public:
 	virtual ~FAnalysisService();
 	virtual TSharedPtr<const IAnalysisSession> Analyze(const TCHAR* SessionUri) override;
 	virtual TSharedPtr<const IAnalysisSession> StartAnalysis(const TCHAR* SessionUri) override;
-	virtual TSharedPtr<const IAnalysisSession> StartAnalysis(const TCHAR* SessionName, TUniquePtr<Trace::IInDataStream>&& DataStream) override;
+	virtual TSharedPtr<const IAnalysisSession> StartAnalysis(const TCHAR* SessionName, TUniquePtr<UE::Trace::IInDataStream>&& DataStream) override;
 
 private:
 	FModuleService& ModuleService;
diff --git a/Engine/Source/Developer/TraceServices/Private/Analyzers/AllocationsAnalysis.h b/Engine/Source/Developer/TraceServices/Private/Analyzers/AllocationsAnalysis.h
--- a/Engine/Source/Developer/TraceServices/Private/Analyzers/AllocationsAnalysis.h
+++ b/Engine/Source/Developer/TraceServices/Private/Analyzers/AllocationsAnalysis.h
@@ -15,7 +15,7 @@ class IAnalysisSession;
 
 ////////////////////////////////////////////////////////////////////////////////
 class FAllocationsAnalyzer
-	: public Trace::IAnalyzer
+	: public UE::Trace::IAnalyzer
 {
 public:
 					FAllocationsAnalyzer(IAnalysisSession& Session);
diff --git a/Engine/Source/Developer/TraceServices/Private/Analyzers/CallstacksAnalysis.h b/Engine/Source/Developer/TraceServices/Private/Analyzers/CallstacksAnalysis.h
--- a/Engine/Source/Developer/TraceServices/Private/Analyzers/CallstacksAnalysis.h
+++ b/Engine/Source/Developer/TraceServices/Private/Analyzers/CallstacksAnalysis.h
@@ -13,7 +13,7 @@ class FCallstacksProvider;
 
 ////////////////////////////////////////////////////////////////////////////////
 class FCallstacksAnalyzer
-	: public Trace::IAnalyzer
+	: public UE::Trace::IAnalyzer
 {
 public:
 					FCallstacksAnalyzer(IAnalysisSession& Session, FCallstacksProvider* Provider);
diff --git a/Engine/Source/Developer/TraceServices/Private/Analyzers/LoadTimeTraceAnalysis.h b/Engine/Source/Developer/TraceServices/Private/Analyzers/LoadTimeTraceAnalysis.h
--- a/Engine/Source/Developer/TraceServices/Private/Analyzers/LoadTimeTraceAnalysis.h
+++ b/Engine/Source/Developer/TraceServices/Private/Analyzers/LoadTimeTraceAnalysis.h
@@ -19,7 +19,7 @@ inline bool operator!=(const FLoadTimeProfilerCpuEvent& Lhs, const FLoadTimeProf
 		Lhs.EventType != Rhs.EventType;
 }
 
-class FAsyncLoadingTraceAnalyzer : public Trace::IAnalyzer
+class FAsyncLoadingTraceAnalyzer : public UE::Trace::IAnalyzer
 {
 public:
 	FAsyncLoadingTraceAnalyzer(IAnalysisSession& Session, FLoadTimeProfilerProvider& LoadTimeProfilerProvider);
diff --git a/Engine/Source/Developer/TraceServices/Private/Common/Utils.h b/Engine/Source/Developer/TraceServices/Private/Common/Utils.h
--- a/Engine/Source/Developer/TraceServices/Private/Common/Utils.h
+++ b/Engine/Source/Developer/TraceServices/Private/Common/Utils.h
@@ -31,7 +31,7 @@ inline int64 DecodeZigZag(const uint8*& BufferPtr)
 }
 
 inline uint32 GetThreadIdField(
-	const Trace::IAnalyzer::FOnEventContext& Context,
+	const UE::Trace::IAnalyzer::FOnEventContext& Context,
 	const ANSICHAR* FieldName="ThreadId")
 {
 	// Trace analysis was changed to be able to provide a suitable id. Prior to
diff --git a/Engine/Source/Developer/TraceServices/Public/TraceServices/AnalysisService.h b/Engine/Source/Developer/TraceServices/Public/TraceServices/AnalysisService.h
--- a/Engine/Source/Developer/TraceServices/Public/TraceServices/AnalysisService.h
+++ b/Engine/Source/Developer/TraceServices/Public/TraceServices/AnalysisService.h
@@ -24,7 +24,7 @@ class IAnalysisService
 public:
 	virtual TSharedPtr<const IAnalysisSession> Analyze(const TCHAR* SessionUri) = 0;
 	virtual TSharedPtr<const IAnalysisSession> StartAnalysis(const TCHAR* SessionUri) = 0;
-	virtual TSharedPtr<const IAnalysisSession> StartAnalysis(const TCHAR* SessionName, TUniquePtr<Trace::IInDataStream>&& DataStream) = 0;
+	virtual TSharedPtr<const IAnalysisSession> StartAnalysis(const TCHAR* SessionName, TUniquePtr<UE::Trace::IInDataStream>&& DataStream) = 0;
 };
 
 } // namespace TraceServices
diff --git a/Engine/Source/Developer/TraceServices/Public/TraceServices/Model/AllocationsProvider.h b/Engine/Source/Developer/TraceServices/Public/TraceServices/Model/AllocationsProvider.h
--- a/Engine/Source/Developer/TraceServices/Public/TraceServices/Model/AllocationsProvider.h
+++ b/Engine/Source/Developer/TraceServices/Public/TraceServices/Model/AllocationsProvider.h
@@ -43,7 +43,7 @@ namespace TraceServices
 
 ////////////////////////////////////////////////////////////////////////////////
 class IAllocationsProvider
-	: public Trace::IProvider
+	: public UE::Trace::IProvider
 {
 public:
 	struct FAllocation
diff --git a/Engine/Source/Developer/TraceServices/Public/TraceServices/Model/AnalysisSession.h b/Engine/Source/Developer/TraceServices/Public/TraceServices/Model/AnalysisSession.h
--- a/Engine/Source/Developer/TraceServices/Public/TraceServices/Model/AnalysisSession.h
+++ b/Engine/Source/Developer/TraceServices/Public/TraceServices/Model/AnalysisSession.h
@@ -7,12 +7,13 @@
 
 class FName;
 
-namespace Trace
-{
+namespace UE {
+namespace Trace {
 
 class IAnalyzer;
 
-}
+} // namespace Trace
+} // namespace UE
 
 namespace TraceServices
 {
@@ -49,7 +50,7 @@ public:
 	virtual const TCHAR* StoreString(const TCHAR* String) = 0;
 	virtual const TCHAR* StoreString(const FStringView& String) = 0;
 	
-	virtual void AddAnalyzer(Trace::IAnalyzer* Analyzer) = 0;
+	virtual void AddAnalyzer(UE::Trace::IAnalyzer* Analyzer) = 0;
 
 	virtual void AddProvider(const FName& Name, IProvider* Provider) = 0;
 	template<typename ProviderType>
diff --git a/Engine/Source/Editor/UnrealEd/Private/CookOnTheFlyServer.cpp b/Engine/Source/Editor/UnrealEd/Private/CookOnTheFlyServer.cpp
--- a/Engine/Source/Editor/UnrealEd/Private/CookOnTheFlyServer.cpp
+++ b/Engine/Source/Editor/UnrealEd/Private/CookOnTheFlyServer.cpp
@@ -2294,7 +2294,7 @@ void UCookOnTheFlyServer::ProcessUnsolicitedPackages()
 }
 
 UE_TRACE_EVENT_BEGIN(UE_CUSTOM_COOKTIMER_LOG, SaveCookedPackage, NoSync)
-	UE_TRACE_EVENT_FIELD(Trace::WideString, PackageName)
+	UE_TRACE_EVENT_FIELD(UE::Trace::WideString, PackageName)
 UE_TRACE_EVENT_END()
 
 void UCookOnTheFlyServer::PumpSaves(UE::Cook::FTickStackData& StackData, uint32 DesiredQueueLength, int32& OutNumPushed, bool& bOutBusy)
diff --git a/Engine/Source/Runtime/Core/Private/Async/TaskGraph.cpp b/Engine/Source/Runtime/Core/Private/Async/TaskGraph.cpp
--- a/Engine/Source/Runtime/Core/Private/Async/TaskGraph.cpp
+++ b/Engine/Source/Runtime/Core/Private/Async/TaskGraph.cpp
@@ -1297,8 +1297,8 @@ public:
 #endif
 			if (GroupName != PrevGroupName)
 			{
-				Trace::ThreadGroupEnd();
-				Trace::ThreadGroupBegin(GroupName);
+				UE::Trace::ThreadGroupEnd();
+				UE::Trace::ThreadGroupBegin(GroupName);
 				PrevGroupName = GroupName;
 			}
 
@@ -1314,7 +1314,7 @@ public:
 			
 			WorkerThreads[ThreadIndex].bAttached = true;
 		}
-		Trace::ThreadGroupEnd();
+		UE::Trace::ThreadGroupEnd();
 	}
 
 	/** 
diff --git a/Engine/Source/Runtime/Core/Private/HAL/LowLevelMemTracker.cpp b/Engine/Source/Runtime/Core/Private/HAL/LowLevelMemTracker.cpp
--- a/Engine/Source/Runtime/Core/Private/HAL/LowLevelMemTracker.cpp
+++ b/Engine/Source/Runtime/Core/Private/HAL/LowLevelMemTracker.cpp
@@ -24,12 +24,12 @@ UE_TRACE_CHANNEL(MemTagChannel, "Memory overview", true)
 UE_TRACE_EVENT_BEGIN(LLM, TagsSpec, NoSync|Important)
 	UE_TRACE_EVENT_FIELD(const void*, TagId)
 	UE_TRACE_EVENT_FIELD(const void*, ParentId)
-	UE_TRACE_EVENT_FIELD(Trace::AnsiString, Name)
+	UE_TRACE_EVENT_FIELD(UE::Trace::AnsiString, Name)
 UE_TRACE_EVENT_END()
 
 UE_TRACE_EVENT_BEGIN(LLM, TrackerSpec, NoSync|Important)
 	UE_TRACE_EVENT_FIELD(uint8, TrackerId)
-	UE_TRACE_EVENT_FIELD(Trace::AnsiString, Name)
+	UE_TRACE_EVENT_FIELD(UE::Trace::AnsiString, Name)
 UE_TRACE_EVENT_END()
 
 UE_TRACE_EVENT_BEGIN(LLM, TagValue)
diff --git a/Engine/Source/Runtime/Core/Private/HAL/ThreadingBase.cpp b/Engine/Source/Runtime/Core/Private/HAL/ThreadingBase.cpp
--- a/Engine/Source/Runtime/Core/Private/HAL/ThreadingBase.cpp
+++ b/Engine/Source/Runtime/Core/Private/HAL/ThreadingBase.cpp
@@ -434,7 +434,7 @@ void FThreadManager::AddThread(uint32 ThreadId, FRunnableThread* Thread)
 	}
 
 	// Note that this must be called from thread being registered.
-	Trace::ThreadRegister(*(Thread->GetThreadName()), Thread->GetThreadID(), SortHint);
+	UE::Trace::ThreadRegister(*(Thread->GetThreadName()), Thread->GetThreadID(), SortHint);
 
 	const bool bIsSingleThreadEnvironment = FPlatformProcess::SupportsMultithreading() == false;
 
@@ -1032,7 +1032,7 @@ public:
 
 	virtual bool Create(uint32 InNumQueuedThreads, uint32 StackSize, EThreadPriority ThreadPriority, const TCHAR* Name) override
 	{
-		Trace::ThreadGroupBegin(Name);
+		UE::Trace::ThreadGroupBegin(Name);
 
 		// Make sure we have synch objects
 		bool bWasSuccessful = true;
@@ -1074,7 +1074,7 @@ public:
 			Destroy();
 		}
 
-		Trace::ThreadGroupEnd();
+		UE::Trace::ThreadGroupEnd();
 		return bWasSuccessful;
 	}
 
diff --git a/Engine/Source/Runtime/Core/Private/Logging/LogTrace.cpp b/Engine/Source/Runtime/Core/Private/Logging/LogTrace.cpp
--- a/Engine/Source/Runtime/Core/Private/Logging/LogTrace.cpp
+++ b/Engine/Source/Runtime/Core/Private/Logging/LogTrace.cpp
@@ -13,7 +13,7 @@ UE_TRACE_CHANNEL(LogChannel)
 UE_TRACE_EVENT_BEGIN(Logging, LogCategory, NoSync|Important)
 	UE_TRACE_EVENT_FIELD(const void*, CategoryPointer)
 	UE_TRACE_EVENT_FIELD(uint8, DefaultVerbosity)
-	UE_TRACE_EVENT_FIELD(Trace::AnsiString, Name)
+	UE_TRACE_EVENT_FIELD(UE::Trace::AnsiString, Name)
 UE_TRACE_EVENT_END()
 
 UE_TRACE_EVENT_BEGIN(Logging, LogMessageSpec, NoSync|Important)
@@ -21,8 +21,8 @@ UE_TRACE_EVENT_BEGIN(Logging, LogMessageSpec, NoSync|Important)
 	UE_TRACE_EVENT_FIELD(const void*, CategoryPointer)
 	UE_TRACE_EVENT_FIELD(int32, Line)
 	UE_TRACE_EVENT_FIELD(uint8, Verbosity)
-	UE_TRACE_EVENT_FIELD(Trace::AnsiString, FileName)
-	UE_TRACE_EVENT_FIELD(Trace::WideString, FormatString)
+	UE_TRACE_EVENT_FIELD(UE::Trace::AnsiString, FileName)
+	UE_TRACE_EVENT_FIELD(UE::Trace::WideString, FormatString)
 UE_TRACE_EVENT_END()
 
 UE_TRACE_EVENT_BEGIN(Logging, LogMessage)
diff --git a/Engine/Source/Runtime/Core/Private/ProfilingDebugging/CountersTrace.cpp b/Engine/Source/Runtime/Core/Private/ProfilingDebugging/CountersTrace.cpp
--- a/Engine/Source/Runtime/Core/Private/ProfilingDebugging/CountersTrace.cpp
+++ b/Engine/Source/Runtime/Core/Private/ProfilingDebugging/CountersTrace.cpp
@@ -12,7 +12,7 @@ UE_TRACE_EVENT_BEGIN(Counters, Spec, NoSync|Important)
 	UE_TRACE_EVENT_FIELD(uint16, Id)
 	UE_TRACE_EVENT_FIELD(uint8, Type)
 	UE_TRACE_EVENT_FIELD(uint8, DisplayHint)
-	UE_TRACE_EVENT_FIELD(Trace::AnsiString, Name)
+	UE_TRACE_EVENT_FIELD(UE::Trace::AnsiString, Name)
 UE_TRACE_EVENT_END()
 
 UE_TRACE_EVENT_BEGIN(Counters, SetValueInt)
diff --git a/Engine/Source/Runtime/Core/Private/ProfilingDebugging/CpuProfilerTrace.cpp b/Engine/Source/Runtime/Core/Private/ProfilingDebugging/CpuProfilerTrace.cpp
--- a/Engine/Source/Runtime/Core/Private/ProfilingDebugging/CpuProfilerTrace.cpp
+++ b/Engine/Source/Runtime/Core/Private/ProfilingDebugging/CpuProfilerTrace.cpp
@@ -16,7 +16,7 @@ UE_TRACE_CHANNEL_DEFINE(CpuChannel)
 
 UE_TRACE_EVENT_BEGIN(CpuProfiler, EventSpec, NoSync|Important)
 	UE_TRACE_EVENT_FIELD(uint32, Id)
-	UE_TRACE_EVENT_FIELD(Trace::AnsiString, Name)
+	UE_TRACE_EVENT_FIELD(UE::Trace::AnsiString, Name)
 UE_TRACE_EVENT_END()
 
 UE_TRACE_EVENT_BEGIN(CpuProfiler, EventBatch, NoSync)
diff --git a/Engine/Source/Runtime/Core/Private/ProfilingDebugging/CsvProfilerTrace.cpp b/Engine/Source/Runtime/Core/Private/ProfilingDebugging/CsvProfilerTrace.cpp
--- a/Engine/Source/Runtime/Core/Private/ProfilingDebugging/CsvProfilerTrace.cpp
+++ b/Engine/Source/Runtime/Core/Private/ProfilingDebugging/CsvProfilerTrace.cpp
@@ -10,19 +10,19 @@
 
 UE_TRACE_EVENT_BEGIN(CsvProfiler, RegisterCategory, NoSync|Important)
 	UE_TRACE_EVENT_FIELD(int32, Index)
-	UE_TRACE_EVENT_FIELD(Trace::AnsiString, Name)
+	UE_TRACE_EVENT_FIELD(UE::Trace::AnsiString, Name)
 UE_TRACE_EVENT_END()
 
 UE_TRACE_EVENT_BEGIN(CsvProfiler, DefineInlineStat, NoSync|Important)
 	UE_TRACE_EVENT_FIELD(uint64, StatId)
 	UE_TRACE_EVENT_FIELD(int32, CategoryIndex)
-	UE_TRACE_EVENT_FIELD(Trace::AnsiString, Name)
+	UE_TRACE_EVENT_FIELD(UE::Trace::AnsiString, Name)
 UE_TRACE_EVENT_END()
 
 UE_TRACE_EVENT_BEGIN(CsvProfiler, DefineDeclaredStat, NoSync|Important)
 	UE_TRACE_EVENT_FIELD(uint64, StatId)
 	UE_TRACE_EVENT_FIELD(int32, CategoryIndex)
-	UE_TRACE_EVENT_FIELD(Trace::AnsiString, Name)
+	UE_TRACE_EVENT_FIELD(UE::Trace::AnsiString, Name)
 UE_TRACE_EVENT_END()
 
 UE_TRACE_EVENT_BEGIN(CsvProfiler, BeginStat)
diff --git a/Engine/Source/Runtime/Core/Private/ProfilingDebugging/MiscTrace.cpp b/Engine/Source/Runtime/Core/Private/ProfilingDebugging/MiscTrace.cpp
--- a/Engine/Source/Runtime/Core/Private/ProfilingDebugging/MiscTrace.cpp
+++ b/Engine/Source/Runtime/Core/Private/ProfilingDebugging/MiscTrace.cpp
@@ -14,8 +14,8 @@ UE_TRACE_CHANNEL(BookmarkChannel)
 UE_TRACE_EVENT_BEGIN(Misc, BookmarkSpec, NoSync|Important)
 	UE_TRACE_EVENT_FIELD(const void*, BookmarkPoint)
 	UE_TRACE_EVENT_FIELD(int32, Line)
-	UE_TRACE_EVENT_FIELD(Trace::WideString, FormatString)
-	UE_TRACE_EVENT_FIELD(Trace::AnsiString, FileName)
+	UE_TRACE_EVENT_FIELD(UE::Trace::WideString, FormatString)
+	UE_TRACE_EVENT_FIELD(UE::Trace::AnsiString, FileName)
 UE_TRACE_EVENT_END()
 
 UE_TRACE_EVENT_BEGIN(Misc, Bookmark)
diff --git a/Engine/Source/Runtime/Core/Private/ProfilingDebugging/TraceAuxiliary.cpp b/Engine/Source/Runtime/Core/Private/ProfilingDebugging/TraceAuxiliary.cpp
--- a/Engine/Source/Runtime/Core/Private/ProfilingDebugging/TraceAuxiliary.cpp
+++ b/Engine/Source/Runtime/Core/Private/ProfilingDebugging/TraceAuxiliary.cpp
@@ -135,7 +135,7 @@ void FTraceAuxiliaryImpl::AddChannel(const TCHAR* Name)
 bool FTraceAuxiliaryImpl::Connect(ETraceConnectType Type, const TCHAR* Parameter)
 {
 	// Connect/write to file. But only if we're not already sending/writing
-	bool bConnected = Trace::IsTracing();
+	bool bConnected = UE::Trace::IsTracing();
 	if (!bConnected)
 	{
 		if (Type == ETraceConnectType::Network)
@@ -176,7 +176,7 @@ bool FTraceAuxiliaryImpl::Connect(ETraceConnectType Type, const TCHAR* Parameter
 ////////////////////////////////////////////////////////////////////////////////
 bool FTraceAuxiliaryImpl::Stop()
 {
-	if (!Trace::Stop())
+	if (!UE::Trace::Stop())
 	{
 		return false;
 	}
@@ -197,12 +197,12 @@ void FTraceAuxiliaryImpl::EnableChannel(FChannel& Channel)
 	// Channel names have been provided by the user and may not exist yet. As
 	// we want to maintain bActive accurately (channels toggles are reference
 	// counted), we will first check Trace knows of the channel.
-	if (!Trace::IsChannel(*Channel.Name))
+	if (!UE::Trace::IsChannel(*Channel.Name))
 	{
 		return;
 	}
 
-	Trace::ToggleChannel(*Channel.Name, true);
+	UE::Trace::ToggleChannel(*Channel.Name, true);
 	Channel.bActive = true;
 }
 
@@ -223,7 +223,7 @@ void FTraceAuxiliaryImpl::DisableChannels()
 		FChannel& Channel = ChannelPair.Value;
 		if (Channel.bActive)
 		{
-			Trace::ToggleChannel(*Channel.Name, false);
+			UE::Trace::ToggleChannel(*Channel.Name, false);
 			Channel.bActive = false;
 		}
 	}
@@ -232,7 +232,7 @@ void FTraceAuxiliaryImpl::DisableChannels()
 ////////////////////////////////////////////////////////////////////////////////
 bool FTraceAuxiliaryImpl::SendToHost(const TCHAR* Host)
 {
-	if (!Trace::SendTo(Host))
+	if (!UE::Trace::SendTo(Host))
 	{
 		UE_LOG(LogCore, Warning, TEXT("Unable to trace to host '%s'"), Host);
 		return false;
@@ -288,7 +288,7 @@ bool FTraceAuxiliaryImpl::WriteToFile(const TCHAR* Path)
 
 	// Finally, tell trace to write the trace to a file.
 	FString NativePath = FileManager.ConvertToAbsolutePathForExternalAppForWrite(*WritePath);
-	if (!Trace::WriteTo(*NativePath))
+	if (!UE::Trace::WriteTo(*NativePath))
 	{
 		UE_LOG(LogCore, Warning, TEXT("Unable to trace to file '%s'"), *WritePath);
 		return false;
@@ -375,9 +375,9 @@ static FAutoConsoleCommand TraceAuxiliaryStopCmd(
 
 ////////////////////////////////////////////////////////////////////////////////
 UE_TRACE_EVENT_BEGIN(Diagnostics, Session2, NoSync|Important)
-	UE_TRACE_EVENT_FIELD(Trace::AnsiString, Platform)
-	UE_TRACE_EVENT_FIELD(Trace::AnsiString, AppName)
-	UE_TRACE_EVENT_FIELD(Trace::WideString, CommandLine)
+	UE_TRACE_EVENT_FIELD(UE::Trace::AnsiString, Platform)
+	UE_TRACE_EVENT_FIELD(UE::Trace::AnsiString, AppName)
+	UE_TRACE_EVENT_FIELD(UE::Trace::WideString, CommandLine)
 	UE_TRACE_EVENT_FIELD(uint8, ConfigurationType)
 	UE_TRACE_EVENT_FIELD(uint8, TargetType)
 UE_TRACE_EVENT_END()
@@ -392,7 +392,7 @@ void FTraceAuxiliary::Initialize(const TCHAR* CommandLine)
 		(UE_ARRAY_COUNT(PREPROCESSOR_TO_STRING(UBT_COMPILED_PLATFORM)) * sizeof(ANSICHAR)) +
 		(UE_ARRAY_COUNT(UE_APP_NAME) * sizeof(ANSICHAR)) +
 		(FCString::Strlen(CommandLine) * sizeof(TCHAR));
-	UE_TRACE_LOG(Diagnostics, Session2, Trace::TraceLogChannel, DataSize)
+	UE_TRACE_LOG(Diagnostics, Session2, UE::Trace::TraceLogChannel, DataSize)
 		<< Session2.Platform(PREPROCESSOR_TO_STRING(UBT_COMPILED_PLATFORM))
 		<< Session2.AppName(UE_APP_NAME)
 		<< Session2.CommandLine(CommandLine)
@@ -400,11 +400,11 @@ void FTraceAuxiliary::Initialize(const TCHAR* CommandLine)
 		<< Session2.TargetType(uint8(FApp::GetBuildTargetType()));
 
 	// Initialize Trace
-	Trace::FInitializeDesc Desc;
+	UE::Trace::FInitializeDesc Desc;
 	Desc.bUseWorkerThread = FPlatformProcess::SupportsMultithreading();
-	Trace::Initialize(Desc);
+	UE::Trace::Initialize(Desc);
 
-	FCoreDelegates::OnEndFrame.AddStatic(Trace::Update);
+	FCoreDelegates::OnEndFrame.AddStatic(UE::Trace::Update);
 	FModuleManager::Get().OnModulesChanged().AddLambda([](FName Name, EModuleChangeReason Reason){
 		if (Reason == EModuleChangeReason::ModuleLoaded)
 		{
@@ -434,8 +434,7 @@ void FTraceAuxiliary::Initialize(const TCHAR* CommandLine)
 		GTraceAuxiliary.Connect(ETraceConnectType::File, nullptr);
 	}
 
-	Trace::ThreadRegister(TEXT("GameThread"), FPlatformTLS::GetCurrentThreadId(), -1);
-
+	UE::Trace::ThreadRegister(TEXT("GameThread"), FPlatformTLS::GetCurrentThreadId(), -1);
 #endif
 }
 
diff --git a/Engine/Source/Runtime/Core/Private/ProfilingDebugging/Windows/ModuleDiagnostics.cpp b/Engine/Source/Runtime/Core/Private/ProfilingDebugging/Windows/ModuleDiagnostics.cpp
--- a/Engine/Source/Runtime/Core/Private/ProfilingDebugging/Windows/ModuleDiagnostics.cpp
+++ b/Engine/Source/Runtime/Core/Private/ProfilingDebugging/Windows/ModuleDiagnostics.cpp
@@ -34,7 +34,7 @@ struct FNtDllFunction
 UE_TRACE_CHANNEL(ModuleChannel, "Module information needed for symbols resolution", true)
 
 UE_TRACE_EVENT_BEGIN(Diagnostics, ModuleLoad, NoSync|Important)
-	UE_TRACE_EVENT_FIELD(Trace::WideString, Name)
+	UE_TRACE_EVENT_FIELD(UE::Trace::WideString, Name)
 	UE_TRACE_EVENT_FIELD(uint32, Base)
 	UE_TRACE_EVENT_FIELD(uint32, Size)
 UE_TRACE_EVENT_END()
@@ -90,7 +90,7 @@ FModuleTrace* FModuleTrace::Get()
 ////////////////////////////////////////////////////////////////////////////////
 void FModuleTrace::Initialize()
 {
-	using namespace Trace;
+	using namespace UE::Trace;
 
 	// Register for DLL load/unload notifications.
 	auto Thunk = [] (ULONG Reason, const void* Data, void* Context)
diff --git a/Engine/Source/Runtime/Core/Public/ProfilingDebugging/CpuProfilerTrace.h b/Engine/Source/Runtime/Core/Public/ProfilingDebugging/CpuProfilerTrace.h
--- a/Engine/Source/Runtime/Core/Public/ProfilingDebugging/CpuProfilerTrace.h
+++ b/Engine/Source/Runtime/Core/Public/ProfilingDebugging/CpuProfilerTrace.h
@@ -68,7 +68,7 @@ struct FCpuProfilerTrace
 
 	struct FEventScope
 	{
-		FEventScope(uint32 InSpecId, const Trace::FChannel& Channel)
+		FEventScope(uint32 InSpecId, const UE::Trace::FChannel& Channel)
 			: bEnabled(Channel | CpuChannel)
 		{
 			if (bEnabled)
@@ -90,7 +90,7 @@ struct FCpuProfilerTrace
 
 	struct FDynamicEventScope
 	{
-		FDynamicEventScope(const ANSICHAR* InEventName, const Trace::FChannel& Channel)
+		FDynamicEventScope(const ANSICHAR* InEventName, const UE::Trace::FChannel& Channel)
 			: bEnabled(Channel | CpuChannel)
 		{
 			if (bEnabled)
@@ -99,7 +99,7 @@ struct FCpuProfilerTrace
 			}
 		}
 
-		FDynamicEventScope(const TCHAR* InEventName, const Trace::FChannel& Channel)
+		FDynamicEventScope(const TCHAR* InEventName, const UE::Trace::FChannel& Channel)
 			: bEnabled(Channel | CpuChannel)
 		{
 			if (bEnabled)
diff --git a/Engine/Source/Runtime/CoreUObject/Private/Serialization/AsyncLoading.cpp b/Engine/Source/Runtime/CoreUObject/Private/Serialization/AsyncLoading.cpp
--- a/Engine/Source/Runtime/CoreUObject/Private/Serialization/AsyncLoading.cpp
+++ b/Engine/Source/Runtime/CoreUObject/Private/Serialization/AsyncLoading.cpp
@@ -4735,14 +4735,14 @@ void FAsyncLoadingThread::StartThread()
 
 	if (!Thread && FAsyncLoadingThreadSettings::Get().bAsyncLoadingThreadEnabled)
 	{
-		Trace::ThreadGroupBegin(TEXT("AsyncLoading"));
+		UE::Trace::ThreadGroupBegin(TEXT("AsyncLoading"));
 
 		UE_LOG(LogStreaming, Log, TEXT("Starting Async Loading Thread."));
 		bThreadStarted = true;
 		FPlatformMisc::MemoryBarrier();
 		Thread = FRunnableThread::Create(this, TEXT("FAsyncLoadingThread"), 0, TPri_Normal);
 
-		Trace::ThreadGroupEnd();
+		UE::Trace::ThreadGroupEnd();
 	}
 }
 
diff --git a/Engine/Source/Runtime/CoreUObject/Private/Serialization/AsyncLoading2.cpp b/Engine/Source/Runtime/CoreUObject/Private/Serialization/AsyncLoading2.cpp
--- a/Engine/Source/Runtime/CoreUObject/Private/Serialization/AsyncLoading2.cpp
+++ b/Engine/Source/Runtime/CoreUObject/Private/Serialization/AsyncLoading2.cpp
@@ -3053,10 +3053,10 @@ FScopedAsyncPackageEvent2::~FScopedAsyncPackageEvent2()
 void FAsyncLoadingThreadWorker::StartThread()
 {
 	LLM_SCOPE(ELLMTag::AsyncLoading);
-	Trace::ThreadGroupBegin(TEXT("AsyncLoading"));
+	UE::Trace::ThreadGroupBegin(TEXT("AsyncLoading"));
 	Thread = FRunnableThread::Create(this, TEXT("FAsyncLoadingThreadWorker"), 0, TPri_Normal);
 	ThreadId = Thread->GetThreadID();
-	Trace::ThreadGroupEnd();
+	UE::Trace::ThreadGroupEnd();
 }
 
 uint32 FAsyncLoadingThreadWorker::Run()
@@ -4652,9 +4652,9 @@ void FAsyncLoadingThread2::StartThread()
 		UE_LOG(LogStreaming, Log, TEXT("Starting Async Loading Thread."));
 		bThreadStarted = true;
 		FPlatformMisc::MemoryBarrier();
-		Trace::ThreadGroupBegin(TEXT("AsyncLoading"));
+		UE::Trace::ThreadGroupBegin(TEXT("AsyncLoading"));
 		Thread = FRunnableThread::Create(this, TEXT("FAsyncLoadingThread"), 0, TPri_Normal);
-		Trace::ThreadGroupEnd();
+		UE::Trace::ThreadGroupEnd();
 	}
 
 	UE_LOG(LogStreaming, Display, TEXT("AsyncLoading2 - Thread Started: %s, IsInitialLoad: %s"),
diff --git a/Engine/Source/Runtime/CoreUObject/Private/UObject/UObjectGlobals.cpp b/Engine/Source/Runtime/CoreUObject/Private/UObject/UObjectGlobals.cpp
--- a/Engine/Source/Runtime/CoreUObject/Private/UObject/UObjectGlobals.cpp
+++ b/Engine/Source/Runtime/CoreUObject/Private/UObject/UObjectGlobals.cpp
@@ -1082,7 +1082,7 @@ public:
 static int32 GGameThreadLoadCounter = 0;
 
 UE_TRACE_EVENT_BEGIN(CUSTOM_LOADTIMER_LOG, LoadPackageInternal, NoSync)
-	UE_TRACE_EVENT_FIELD(Trace::WideString, PackageName)
+	UE_TRACE_EVENT_FIELD(UE::Trace::WideString, PackageName)
 UE_TRACE_EVENT_END()
 
 UPackage* LoadPackageInternal(UPackage* InOuter, const TCHAR* InLongPackageNameOrFilename, uint32 LoadFlags, FLinkerLoad* ImportLinker, FArchive* InReaderOverride, const FLinkerInstancingContext* InstancingContext)
diff --git a/Engine/Source/Runtime/Engine/Private/Animation/AnimTrace.cpp b/Engine/Source/Runtime/Engine/Private/Animation/AnimTrace.cpp
--- a/Engine/Source/Runtime/Engine/Private/Animation/AnimTrace.cpp
+++ b/Engine/Source/Runtime/Engine/Private/Animation/AnimTrace.cpp
@@ -191,7 +191,7 @@ UE_TRACE_EVENT_END()
 
 UE_TRACE_EVENT_BEGIN(Animation, Name, NoSync|Important)
 	UE_TRACE_EVENT_FIELD(uint32, Id)
-	UE_TRACE_EVENT_FIELD(Trace::WideString, Name)
+	UE_TRACE_EVENT_FIELD(UE::Trace::WideString, Name)
 UE_TRACE_EVENT_END()
 
 UE_TRACE_EVENT_BEGIN(Animation, Notify)
diff --git a/Engine/Source/Runtime/Net/Core/Private/Net/Core/Trace/NetTraceInternal.cpp b/Engine/Source/Runtime/Net/Core/Private/Net/Core/Trace/NetTraceInternal.cpp
--- a/Engine/Source/Runtime/Net/Core/Private/Net/Core/Trace/NetTraceInternal.cpp
+++ b/Engine/Source/Runtime/Net/Core/Private/Net/Core/Trace/NetTraceInternal.cpp
@@ -73,7 +73,7 @@ void FNetTrace::SetTraceVerbosity(uint32 Verbosity)
 	// Enable
 	if (!GetTraceVerbosity() && NewVerbosity)
 	{
-		Trace::ToggleChannel(TEXT("NetChannel"), true);
+		UE::Trace::ToggleChannel(TEXT("NetChannel"), true);
 
 		FNetTraceInternal::Reporter::ReportInitEvent(FNetTraceInternal::NetTraceVersion);
 	}
@@ -81,7 +81,7 @@ void FNetTrace::SetTraceVerbosity(uint32 Verbosity)
 	{
 		if (FNetTraceInternal::ThreadBuffer)
 		{
-			Trace::ToggleChannel(TEXT("NetChannel"), false);
+			UE::Trace::ToggleChannel(TEXT("NetChannel"), false);
 
 			delete FNetTraceInternal::ThreadBuffer;
 			FNetTraceInternal::ThreadBuffer = nullptr;
diff --git a/Engine/Source/Runtime/RHI/Public/RHICommandList.h b/Engine/Source/Runtime/RHI/Public/RHICommandList.h
--- a/Engine/Source/Runtime/RHI/Public/RHICommandList.h
+++ b/Engine/Source/Runtime/RHI/Public/RHICommandList.h
@@ -86,7 +86,7 @@ extern RHI_API uint32 GWorkingRHIThreadStartCycles;
 /** How many cycles the from sampling input to the frame being flipped. */
 extern RHI_API uint64 GInputLatencyTime;
 
-/*Trace::FChannel& FORCEINLINE GetRHICommandsChannel() 
+/*UE::Trace::FChannel& FORCEINLINE GetRHICommandsChannel() 
 {
 
 }*/
diff --git a/Engine/Source/Runtime/RenderCore/Private/RenderingThread.cpp b/Engine/Source/Runtime/RenderCore/Private/RenderingThread.cpp
--- a/Engine/Source/Runtime/RenderCore/Private/RenderingThread.cpp
+++ b/Engine/Source/Runtime/RenderCore/Private/RenderingThread.cpp
@@ -329,12 +329,12 @@ public:
 
 	void Start()
 	{
-		Trace::ThreadGroupBegin(TEXT("Render"));
+		UE::Trace::ThreadGroupBegin(TEXT("Render"));
 		Thread = FRunnableThread::Create(this, TEXT("RHIThread"), 512 * 1024, FPlatformAffinity::GetRHIThreadPriority(),
 			FPlatformAffinity::GetRHIThreadMask(), FPlatformAffinity::GetRHIThreadFlags()
 			);
 		check(Thread);
-		Trace::ThreadGroupEnd();
+		UE::Trace::ThreadGroupEnd();
 	}
 };
 
@@ -766,12 +766,12 @@ PRAGMA_ENABLE_DEPRECATION_WARNINGS
 	// Create the rendering thread.
 	GRenderingThreadRunnable = new FRenderingThread();
 
-	Trace::ThreadGroupBegin(TEXT("Render"));
+	UE::Trace::ThreadGroupBegin(TEXT("Render"));
 PRAGMA_DISABLE_DEPRECATION_WARNINGS
 	GRenderingThread = 
 PRAGMA_ENABLE_DEPRECATION_WARNINGS
 		FRunnableThread::Create(GRenderingThreadRunnable, *BuildRenderingThreadName(ThreadCount), 0, FPlatformAffinity::GetRenderingThreadPriority(), FPlatformAffinity::GetRenderingThreadMask(), FPlatformAffinity::GetRenderingThreadFlags());
-	Trace::ThreadGroupEnd();
+	UE::Trace::ThreadGroupEnd();
 
 	// Wait for render thread to have taskgraph bound before we dispatch any tasks for it.
 	((FRenderingThread*)GRenderingThreadRunnable)->TaskGraphBoundSyncEvent->Wait();
@@ -788,9 +788,9 @@ PRAGMA_ENABLE_DEPRECATION_WARNINGS
 	// Create the rendering thread heartbeat
 	GRenderingThreadRunnableHeartbeat = new FRenderingThreadTickHeartbeat();
 
-	Trace::ThreadGroupBegin(TEXT("Render"));
+	UE::Trace::ThreadGroupBegin(TEXT("Render"));
 	GRenderingThreadHeartbeat = FRunnableThread::Create(GRenderingThreadRunnableHeartbeat, *FString::Printf(TEXT("RTHeartBeat %d"), ThreadCount), 16 * 1024, TPri_AboveNormal, FPlatformAffinity::GetRTHeartBeatMask());
-	Trace::ThreadGroupEnd();
+	UE::Trace::ThreadGroupEnd();
 
 	ThreadCount++;
 
diff --git a/Engine/Source/Runtime/SlateCore/Private/Trace/SlateTrace.cpp b/Engine/Source/Runtime/SlateCore/Private/Trace/SlateTrace.cpp
--- a/Engine/Source/Runtime/SlateCore/Private/Trace/SlateTrace.cpp
+++ b/Engine/Source/Runtime/SlateCore/Private/Trace/SlateTrace.cpp
@@ -51,8 +51,8 @@ UE_TRACE_EVENT_END()
 UE_TRACE_EVENT_BEGIN(SlateTrace, WidgetInfo)
 	UE_TRACE_EVENT_FIELD(uint64, Cycle)
 	UE_TRACE_EVENT_FIELD(uint64, WidgetId)				// Created/Updated widget unique ID.
-	UE_TRACE_EVENT_FIELD(Trace::WideString, Path)		// FReflectionMetaData::GetWidgetPath
-	UE_TRACE_EVENT_FIELD(Trace::WideString, DebugInfo)	// FReflectionMetaData::GetWidgetDebugInfo
+	UE_TRACE_EVENT_FIELD(UE::Trace::WideString, Path)	// FReflectionMetaData::GetWidgetPath
+	UE_TRACE_EVENT_FIELD(UE::Trace::WideString, DebugInfo)// FReflectionMetaData::GetWidgetDebugInfo
 UE_TRACE_EVENT_END()
 
 UE_TRACE_EVENT_BEGIN(SlateTrace, RemoveWidget)
@@ -71,7 +71,7 @@ UE_TRACE_EVENT_BEGIN(SlateTrace, WidgetInvalidated)
 	UE_TRACE_EVENT_FIELD(uint64, WidgetId)					// Invalidated widget unique ID.
 	UE_TRACE_EVENT_FIELD(uint64, InvestigatorId)			// Widget unique ID that investigated the invalidation.
 	UE_TRACE_EVENT_FIELD(uint8, InvalidateWidgetReason)		// The reason of the invalidation. (EInvalidateWidgetReason)
-	UE_TRACE_EVENT_FIELD(Trace::WideString, ScriptTrace)	// Optional script trace for root widget invalidations
+	UE_TRACE_EVENT_FIELD(UE::Trace::WideString, ScriptTrace)// Optional script trace for root widget invalidations
 	UE_TRACE_EVENT_FIELD(uint64[], Callstack)				// Optional callstack for root widget invalidations
 	UE_TRACE_EVENT_FIELD(uint32, ProcessId)					// Optional proccess ID where the invalidation occureed.
 UE_TRACE_EVENT_END()
diff --git a/Engine/Source/Runtime/TraceLog/Private/Trace/BlockPool.cpp b/Engine/Source/Runtime/TraceLog/Private/Trace/BlockPool.cpp
--- a/Engine/Source/Runtime/TraceLog/Private/Trace/BlockPool.cpp
+++ b/Engine/Source/Runtime/TraceLog/Private/Trace/BlockPool.cpp
@@ -8,6 +8,7 @@
 #include "Trace/Detail/Atomic.h"
 #include "Trace/Detail/Writer.inl"
 
+namespace UE {
 namespace Trace {
 namespace Private {
 
@@ -182,5 +183,6 @@ void Writer_ShutdownPool()
 
 } // namespace Private
 } // namespace Trace
+} // namespace UE
 
 #endif // UE_TRACE_ENABLED
diff --git a/Engine/Source/Runtime/TraceLog/Private/Trace/Channel.cpp b/Engine/Source/Runtime/TraceLog/Private/Trace/Channel.cpp
--- a/Engine/Source/Runtime/TraceLog/Private/Trace/Channel.cpp
+++ b/Engine/Source/Runtime/TraceLog/Private/Trace/Channel.cpp
@@ -11,6 +11,7 @@
 
 #if UE_TRACE_ENABLED
 
+namespace UE {
 namespace Trace {
 
 ////////////////////////////////////////////////////////////////////////////////
@@ -28,7 +29,7 @@ UE_TRACE_EVENT_BEGIN(Trace, ChannelAnnounce, NoSync|Important)
 	UE_TRACE_EVENT_FIELD(uint32, Id)
 	UE_TRACE_EVENT_FIELD(bool, IsEnabled)
 	UE_TRACE_EVENT_FIELD(bool, ReadOnly)
-	UE_TRACE_EVENT_FIELD(Trace::AnsiString, Name)
+	UE_TRACE_EVENT_FIELD(AnsiString, Name)
 UE_TRACE_EVENT_END()
 
 UE_TRACE_EVENT_BEGIN(Trace, ChannelToggle, NoSync|Important)
@@ -252,5 +253,6 @@ bool FChannel::Toggle(const ANSICHAR* ChannelName, bool bEnabled)
 }
 
 } // namespace Trace
+} // namespace UE
 
 #endif // UE_TRACE_ENABLED
diff --git a/Engine/Source/Runtime/TraceLog/Private/Trace/Codec.cpp b/Engine/Source/Runtime/TraceLog/Private/Trace/Codec.cpp
--- a/Engine/Source/Runtime/TraceLog/Private/Trace/Codec.cpp
+++ b/Engine/Source/Runtime/TraceLog/Private/Trace/Codec.cpp
@@ -17,13 +17,14 @@ THIRD_PARTY_INCLUDES_START
 #endif
 THIRD_PARTY_INCLUDES_END
 
+namespace UE {
 namespace Trace {
 namespace Private {
 
 ////////////////////////////////////////////////////////////////////////////////
 int32 Encode(const void* Src, int32 SrcSize, void* Dest, int32 DestSize)
 {
-	return Trace::LZ4_compress_fast(
+	return ::Trace::LZ4_compress_fast(
 		(const char*)Src,
 		(char*)Dest,
 		SrcSize,
@@ -41,8 +42,9 @@ uint32 GetEncodeMaxSize(uint32 InputSize)
 ////////////////////////////////////////////////////////////////////////////////
 TRACELOG_API int32 Decode(const void* Src, int32 SrcSize, void* Dest, int32 DestSize)
 {
-	return Trace::LZ4_decompress_safe((const char*)Src, (char*)Dest, SrcSize, DestSize);
+	return ::Trace::LZ4_decompress_safe((const char*)Src, (char*)Dest, SrcSize, DestSize);
 }
 
 } // namespace Private
 } // namespace Trace
+} // namespace UE
diff --git a/Engine/Source/Runtime/TraceLog/Private/Trace/Control.cpp b/Engine/Source/Runtime/TraceLog/Private/Trace/Control.cpp
--- a/Engine/Source/Runtime/TraceLog/Private/Trace/Control.cpp
+++ b/Engine/Source/Runtime/TraceLog/Private/Trace/Control.cpp
@@ -9,6 +9,7 @@
 #include "Misc/CString.h"
 #include "Templates/UnrealTemplate.h"
 
+namespace UE {
 namespace Trace {
 namespace Private {
 
@@ -335,5 +336,6 @@ void Writer_ShutdownControl()
 
 } // namespace Private
 } // namespace Trace
+} // namespace UE
 
 #endif // UE_TRACE_ENABLED
diff --git a/Engine/Source/Runtime/TraceLog/Private/Trace/Trace.cpp b/Engine/Source/Runtime/TraceLog/Private/Trace/Trace.cpp
--- a/Engine/Source/Runtime/TraceLog/Private/Trace/Trace.cpp
+++ b/Engine/Source/Runtime/TraceLog/Private/Trace/Trace.cpp
@@ -8,8 +8,8 @@
 #include "Trace/Detail/Channel.h"
 
 
-namespace Trace
-{
+namespace UE {
+namespace Trace {
 
 namespace Private
 {
@@ -125,11 +125,11 @@ UE_TRACE_EVENT_BEGIN($Trace, ThreadInfo, NoSync|Important)
 	UE_TRACE_EVENT_FIELD(uint32, ThreadId)
 	UE_TRACE_EVENT_FIELD(uint32, SystemId)
 	UE_TRACE_EVENT_FIELD(int32, SortHint)
-	UE_TRACE_EVENT_FIELD(Trace::AnsiString, Name)
+	UE_TRACE_EVENT_FIELD(AnsiString, Name)
 UE_TRACE_EVENT_END()
 
 UE_TRACE_EVENT_BEGIN($Trace, ThreadGroupBegin, NoSync|Important)
-	UE_TRACE_EVENT_FIELD(Trace::AnsiString, Name)
+	UE_TRACE_EVENT_FIELD(AnsiString, Name)
 UE_TRACE_EVENT_END()
 
 UE_TRACE_EVENT_BEGIN($Trace, ThreadGroupEnd, NoSync|Important)
@@ -162,6 +162,7 @@ void ThreadGroupEnd()
 }
 
 } // namespace Trace
+} // namespace UE
 
 #else
 
diff --git a/Engine/Source/Runtime/TraceLog/Private/Trace/Writer.cpp b/Engine/Source/Runtime/TraceLog/Private/Trace/Writer.cpp
--- a/Engine/Source/Runtime/TraceLog/Private/Trace/Writer.cpp
+++ b/Engine/Source/Runtime/TraceLog/Private/Trace/Writer.cpp
@@ -25,6 +25,7 @@
 #	define TRACE_PRIVATE_STOMP 0
 #endif
 
+namespace UE {
 namespace Trace {
 namespace Private {
 
@@ -439,7 +440,7 @@ static void Writer_WorkerUpdate()
 ////////////////////////////////////////////////////////////////////////////////
 static void Writer_WorkerThread()
 {
-	Trace::ThreadRegister(TEXT("Trace"), 0, INT_MAX);
+	ThreadRegister(TEXT("Trace"), 0, INT_MAX);
 
 	Writer_UpdateControl();
 
@@ -650,5 +651,6 @@ bool Writer_Stop()
 
 } // namespace Private
 } // namespace Trace
+} // namespace UE
 
 #endif // UE_TRACE_ENABLED
diff --git a/Engine/Source/Runtime/TraceLog/Public/Trace/Detail/Field.h b/Engine/Source/Runtime/TraceLog/Public/Trace/Detail/Field.h
--- a/Engine/Source/Runtime/TraceLog/Public/Trace/Detail/Field.h
+++ b/Engine/Source/Runtime/TraceLog/Public/Trace/Detail/Field.h
@@ -16,8 +16,11 @@
  * where and partially how to implement them */
 #define STATICALLY_SIZED_ARRAY_FIELDS_SUPPORT 0
 
+namespace UE {
 namespace Trace {
-namespace Private {
+
+namespace Private
+{
 
 ////////////////////////////////////////////////////////////////////////////////
 UE_TRACE_API void Field_WriteAuxData(uint32, const uint8*, int32);
@@ -65,8 +68,8 @@ struct TFieldType<T[N]>
 };
 #endif // STATICALLY_SIZED_ARRAY_FIELDS_SUPPORT
 
-template <> struct TFieldType<Trace::AnsiString> { enum { Tid  = int(EFieldType::AnsiString), Size = 0, }; };
-template <> struct TFieldType<Trace::WideString> { enum { Tid  = int(EFieldType::WideString), Size = 0, }; };
+template <> struct TFieldType<AnsiString> { enum { Tid  = int(EFieldType::AnsiString), Size = 0, }; };
+template <> struct TFieldType<WideString> { enum { Tid  = int(EFieldType::WideString), Size = 0, }; };
 
 
 
@@ -152,16 +155,16 @@ struct TField<InIndex, InOffset, Type[Count]>
 
 ////////////////////////////////////////////////////////////////////////////////
 template <int InIndex, int InOffset>
-struct TField<InIndex, InOffset, Trace::AnsiString>
+struct TField<InIndex, InOffset, AnsiString>
 {
-	TRACE_PRIVATE_FIELD(InIndex + int(EIndexPack::AuxFieldCounter), InOffset, Trace::AnsiString);
+	TRACE_PRIVATE_FIELD(InIndex + int(EIndexPack::AuxFieldCounter), InOffset, AnsiString);
 };
 
 ////////////////////////////////////////////////////////////////////////////////
 template <int InIndex, int InOffset>
-struct TField<InIndex, InOffset, Trace::WideString>
+struct TField<InIndex, InOffset, WideString>
 {
-	TRACE_PRIVATE_FIELD(InIndex + int(EIndexPack::AuxFieldCounter), InOffset, Trace::WideString);
+	TRACE_PRIVATE_FIELD(InIndex + int(EIndexPack::AuxFieldCounter), InOffset, WideString);
 };
 
 ////////////////////////////////////////////////////////////////////////////////
@@ -199,5 +202,6 @@ struct TField<0, InOffset, Attachment>
 };
 
 } // namespace Trace
+} // namespace UE
 
 #endif // UE_TRACE_ENABLED
diff --git a/Engine/Source/Runtime/TraceLog/Public/Trace/Detail/Protocols/Protocol0.h b/Engine/Source/Runtime/TraceLog/Public/Trace/Detail/Protocols/Protocol0.h
--- a/Engine/Source/Runtime/TraceLog/Public/Trace/Detail/Protocols/Protocol0.h
+++ b/Engine/Source/Runtime/TraceLog/Public/Trace/Detail/Protocols/Protocol0.h
@@ -2,8 +2,8 @@
 
 #pragma once
 
-namespace Trace
-{
+namespace UE {
+namespace Trace {
 
 #if defined(TRACE_PRIVATE_PROTOCOL_0)
 inline
@@ -99,5 +99,5 @@ struct FEventHeader
 };
 
 } // namespace Protocol0
-
 } // namespace Trace
+} // namespace UE
