Homogenised return type of Thread::GetInfo() across all its overloads.

#rnx

diff --git a/Engine/Source/Developer/TraceAnalysis/Private/Analysis/Engine.cpp b/Engine/Source/Developer/TraceAnalysis/Private/Analysis/Engine.cpp
--- a/Engine/Source/Developer/TraceAnalysis/Private/Analysis/Engine.cpp
+++ b/Engine/Source/Developer/TraceAnalysis/Private/Analysis/Engine.cpp
@@ -75,7 +75,7 @@ FThreads::FInfo* FThreads::GetInfo()
 }
 
 ////////////////////////////////////////////////////////////////////////////////
-FThreads::FInfo& FThreads::GetInfo(uint32 ThreadId)
+FThreads::FInfo* FThreads::GetInfo(uint32 ThreadId)
 {
 	LastGetInfoId = ThreadId;
 
@@ -89,7 +89,7 @@ FThreads::FInfo& FThreads::GetInfo(uint32 ThreadId)
 	{
 		Info.ThreadId = ThreadId;
 	}
-	return Info;
+	return &Info;
 }
 
 ////////////////////////////////////////////////////////////////////////////////
@@ -1426,16 +1426,16 @@ bool FAnalysisEngine::OnDataProtocol2()
 				continue;
 			}
 
-			FThreads::FInfo& ThreadInfo = Threads.GetInfo(RotaItem.ThreadId);
+			FThreads::FInfo* ThreadInfo = Threads.GetInfo(RotaItem.ThreadId);
 
 			uint32 AvailableSerial;
 			if (ProtocolVersion == UE::Trace::Protocol4::EProtocol::Id)
 			{
-				AvailableSerial = OnDataProtocol4(*(RotaItem.Reader), ThreadInfo);
+				AvailableSerial = OnDataProtocol4(*(RotaItem.Reader), *ThreadInfo);
 			}
 			else
 			{
-				AvailableSerial = OnDataProtocol2(*(RotaItem.Reader), ThreadInfo);
+				AvailableSerial = OnDataProtocol2(*(RotaItem.Reader), *ThreadInfo);
 			}
 
 			if (int32(AvailableSerial) >= 0)
diff --git a/Engine/Source/Developer/TraceAnalysis/Private/Analysis/Engine.h b/Engine/Source/Developer/TraceAnalysis/Private/Analysis/Engine.h
--- a/Engine/Source/Developer/TraceAnalysis/Private/Analysis/Engine.h
+++ b/Engine/Source/Developer/TraceAnalysis/Private/Analysis/Engine.h
@@ -28,7 +28,7 @@ public:
 
 						FThreads();
 	FInfo*				GetInfo();
-	FInfo&				GetInfo(uint32 ThreadId);
+	FInfo*				GetInfo(uint32 ThreadId);
 	void				SetGroupName(const ANSICHAR* InGroupName, uint32 Length);
 	const TArray<uint8>*GetGroupName() const;
 
