Trace's analysis' FThreads::GetInfo() was always returning thread 0 (the important "thread"). Since they were marked as important, thread-info events have no longer been sent on the thread the present but instead carry the thread ID in a field.

#rnx

diff --git a/Engine/Source/Developer/TraceAnalysis/Private/Analysis/Engine.cpp b/Engine/Source/Developer/TraceAnalysis/Private/Analysis/Engine.cpp
--- a/Engine/Source/Developer/TraceAnalysis/Private/Analysis/Engine.cpp
+++ b/Engine/Source/Developer/TraceAnalysis/Private/Analysis/Engine.cpp
@@ -982,11 +982,24 @@ void FAnalysisEngine::OnThreadInfoInternal(const FOnEventContext& Context)
 {
 	const FEventData& EventData = Context.EventData;
 
-	FThreads::FInfo* ThreadInfo = Threads.GetInfo();
+	FThreads::FInfo* ThreadInfo;
+	uint32 ThreadId = EventData.GetValue<uint32>("ThreadId", ~0u);
+	if (ThreadId != ~0u)
+	{
+		// Post important-events; the thread-info event is not on the thread it
+		// represents anymore. Fortunately the thread-id is traced now.
+		ThreadInfo = Threads.GetInfo(ThreadId);
+	}
+	else
+	{
+		ThreadInfo = Threads.GetInfo();
+	}
+
 	if (ThreadInfo == nullptr)
 	{
 		return;
 	}
+
 	ThreadInfo->SystemId = EventData.GetValue<uint32>("SystemId");
 	ThreadInfo->SortHint = EventData.GetValue<int32>("SortHint");
 	
@@ -1004,10 +1017,6 @@ void FAnalysisEngine::OnThreadInfoInternal(const FOnEventContext& Context)
 		}
 	}
 
-	// Implicit ThreadId detail is lost as these events are "important". Fake it
-	uint32 PrevThreadId = ThreadInfo->ThreadId;
-	ThreadInfo->ThreadId = EventData.GetValue<uint32>("ThreadId", PrevThreadId);
-
 	const auto* OuterInfo = (FThreadInfo*)ThreadInfo;
 	for (uint16 i = 0, n = Analyzers.Num(); i < n; ++i)
 	{
@@ -1016,8 +1025,6 @@ void FAnalysisEngine::OnThreadInfoInternal(const FOnEventContext& Context)
 			Analyzer->OnThreadInfo(*OuterInfo);
 		}
 	}
-
-	ThreadInfo->ThreadId = PrevThreadId;
 }
 
 ////////////////////////////////////////////////////////////////////////////////
