Fixed pointer invalidation when new thread events were received.

Events marked as important are associated internally to a pseudo thread with id 0 [A]. New-thread events are considered important and thus appear as if they were traced on thread 0. They will add thread N which will modify the Infos array. If someone was holding a A& = Infos[0] when Infos[N] was added then A can become stale. To combat this edge-case we'll thread 0 as a special case so it doesn't invalidate itself with its own events.

#rnx

diff --git a/Engine/Source/Developer/TraceAnalysis/Private/Analysis/Engine.cpp b/Engine/Source/Developer/TraceAnalysis/Private/Analysis/Engine.cpp
--- a/Engine/Source/Developer/TraceAnalysis/Private/Analysis/Engine.cpp
+++ b/Engine/Source/Developer/TraceAnalysis/Private/Analysis/Engine.cpp
@@ -66,12 +66,7 @@ FThreads::FThreads()
 ////////////////////////////////////////////////////////////////////////////////
 FThreads::FInfo* FThreads::GetInfo()
 {
-	if (LastGetInfoId >= uint32(Infos.Num()))
-	{
-		return nullptr;
-	}
-
-	return Infos.GetData() + LastGetInfoId;
+	return GetInfo(LastGetInfoId);
 }
 
 ////////////////////////////////////////////////////////////////////////////////
@@ -79,12 +74,24 @@ FThreads::FInfo* FThreads::GetInfo(uint32 ThreadId)
 {
 	LastGetInfoId = ThreadId;
 
+	// Events marked as important are associated internally to a pseudo thread
+	// with id 0 [A]. New-thread events are considered important and thus appear
+	// as if they were traced on thread 0. They will add thread N which will
+	// modify the Infos array. If someone was holding a A& = Infos[0] when
+	// Infos[N] was added then A can become stale. To combat this edge-case
+	// we'll thread 0 as a special case so it doesn't invalidate itself with its
+	// own events.
+	if (ThreadId == 0)
+	{
+		return &ImportantInfo;
+	}
+
 	if (ThreadId >= uint32(Infos.Num()))
 	{
-		Infos.SetNum(ThreadId + 1);
+		Infos.SetNum(ThreadId);
 	}
 
-	FInfo& Info = Infos[ThreadId];
+	FInfo& Info = Infos[ThreadId - 1];
 	if (Info.ThreadId == ~0u)
 	{
 		Info.ThreadId = ThreadId;
diff --git a/Engine/Source/Developer/TraceAnalysis/Private/Analysis/Engine.h b/Engine/Source/Developer/TraceAnalysis/Private/Analysis/Engine.h
--- a/Engine/Source/Developer/TraceAnalysis/Private/Analysis/Engine.h
+++ b/Engine/Source/Developer/TraceAnalysis/Private/Analysis/Engine.h
@@ -34,6 +34,7 @@ public:
 
 private:
 	uint32				LastGetInfoId = ~0u;
+	FInfo				ImportantInfo;
 	TArray<FInfo>		Infos;
 	TArray<uint8>		GroupName;
 };
