Moved event data pointer to the start of an internal struct so it is a little bit easier to find from the outside.

#rnx

diff --git a/Engine/Source/Developer/TraceAnalysis/Private/Analysis/Engine.cpp b/Engine/Source/Developer/TraceAnalysis/Private/Analysis/Engine.cpp
--- a/Engine/Source/Developer/TraceAnalysis/Private/Analysis/Engine.cpp
+++ b/Engine/Source/Developer/TraceAnalysis/Private/Analysis/Engine.cpp
@@ -506,9 +506,9 @@ const void* IAnalyzer::FArrayReader::GetImpl(uint32 Index, int16& SizeAndType) c
 struct FAnalysisEngine::FEventDataInfo
 {
 	const FAuxData*		GetAuxData(uint32 FieldIndex) const;
+	const uint8*		Ptr;
 	const FDispatch&	Dispatch;
 	FAuxDataCollector*	AuxCollector;
-	const uint8*		Ptr;
 	uint16				Size;
 };
 
@@ -1379,9 +1379,9 @@ bool FAnalysisEngine::OnDataProtocol0()
 		}
 
 		FEventDataInfo EventDataInfo = {
+			Header->EventData,
 			*Dispatch,
 			nullptr,
-			Header->EventData,
 			Header->Size
 		};
 
@@ -1626,9 +1626,9 @@ int32 FAnalysisEngine::OnDataProtocol2(FStreamReader& Reader, FThreads::FInfo& T
 		}
 
 		FEventDataInfo EventDataInfo = {
+			(const uint8*)Header + BlockSize - Header->Size,
 			*Dispatch,
 			&AuxCollector,
-			(const uint8*)Header + BlockSize - Header->Size,
 			Header->Size,
 		};
 
@@ -1770,7 +1770,7 @@ int32 FAnalysisEngine::OnDataProtocol4Known(
 			uint32 RouteIndex = ThreadInfo.ScopeRoutes.Pop(false);
 
 			FDispatch EmptyDispatch = {};
-			FEventDataInfo EmptyEventInfo = { EmptyDispatch };
+			FEventDataInfo EmptyEventInfo = { nullptr, EmptyDispatch };
 
 			FOnEventContext Context = {
 				(const FThreadInfo&)ThreadInfo,
@@ -1917,9 +1917,9 @@ int32 FAnalysisEngine::OnDataProtocol4Impl(
 
 	// Sent the event to subscribed analyzers
 	FEventDataInfo EventDataInfo = {
+		(const uint8*)Header + BlockSize - Header->Size,
 		*Dispatch,
 		&AuxCollector,
-		(const uint8*)Header + BlockSize - Header->Size,
 		Header->Size,
 	};
 
