An alternative API for reading fields' values. A opaque handle can be retrieved when an event is first described instead of hashing a string and traversing a table for each GetValue() call. Bast-case can be an instruction or two instead of much work.

#rnx

diff --git a/Engine/Source/Developer/TraceAnalysis/Private/Analysis/Engine.cpp b/Engine/Source/Developer/TraceAnalysis/Private/Analysis/Engine.cpp
--- a/Engine/Source/Developer/TraceAnalysis/Private/Analysis/Engine.cpp
+++ b/Engine/Source/Developer/TraceAnalysis/Private/Analysis/Engine.cpp
@@ -435,6 +435,23 @@ const IAnalyzer::FEventFieldInfo* IAnalyzer::FEventTypeInfo::GetFieldInfo(uint32
 	return (const IAnalyzer::FEventFieldInfo*)(Inner->Fields + Index);
 }
 
+////////////////////////////////////////////////////////////////////////////////
+IAnalyzer::FEventFieldHandle IAnalyzer::FEventTypeInfo::GetFieldHandleImpl(
+	const ANSICHAR* FieldName,
+	int16& SizeAndType) const
+{
+	const auto* Inner = (const FAnalysisEngine::FDispatch*)this;
+	int32 Index = Inner->GetFieldIndex(FieldName);
+	if (Index < 0)
+	{
+		return { -1 };
+	}
+
+	const FAnalysisEngine::FDispatch::FField& Field = Inner->Fields[Index];
+	SizeAndType = Field.SizeAndType;
+	return { Field.Offset };
+}
+
 
 
 ////////////////////////////////////////////////////////////////////////////////
diff --git a/Engine/Source/Developer/TraceAnalysis/Public/Trace/Analyzer.h b/Engine/Source/Developer/TraceAnalysis/Public/Trace/Analyzer.h
--- a/Engine/Source/Developer/TraceAnalysis/Public/Trace/Analyzer.h
+++ b/Engine/Source/Developer/TraceAnalysis/Public/Trace/Analyzer.h
@@ -6,6 +6,8 @@
 #include "Containers/ArrayView.h"
 #include "Containers/StringView.h"
 
+#include <type_traits>
+
 namespace UE {
 namespace Trace {
 
@@ -61,6 +63,12 @@ public:
 		bool IsArray() const;
 	};
 
+	struct FEventFieldHandle
+	{
+		bool	IsValid() const { return Detail >= 0; }
+		int32	Detail;
+	};
+
 	struct TRACEANALYSIS_API FEventTypeInfo
 	{
 		/** Each event is assigned a unique ID when logged. Not that this is not
@@ -78,6 +86,16 @@ public:
 
 		/** By-index access to fields' type information. */
 		const FEventFieldInfo* GetFieldInfo(uint32 Index) const;
+
+		/** Returns a handle that can used to access events' fields. There is 
+		 * loose validation via ValueType, but one should still exercise caution
+		 * when reading fields with handles.
+		 * @param ValueType The intended type that the field will be interpreted as */
+		template <typename ValueType>
+		FEventFieldHandle GetFieldHandle(const ANSICHAR* FieldName) const;
+
+	private:
+		FEventFieldHandle GetFieldHandleImpl(const ANSICHAR*, int16&) const;
 	};
 
 	struct TRACEANALYSIS_API FArrayReader
@@ -112,6 +130,7 @@ public:
 		 * @param Default Return this value if the given field was not found.
 		 * @return Value of the field (coerced to ValueType) if found, otherwise 0. */
 		template <typename ValueType> ValueType GetValue(const ANSICHAR* FieldName, ValueType Default=ValueType(0)) const;
+		template <typename ValueType> ValueType GetValue(FEventFieldHandle FieldHandle) const;
 
 		/** Returns an object for reading data from an array-type field. A valid
 		 * array reader object will always be return even if no field matching the
@@ -251,6 +270,23 @@ private:
 	template <typename ValueType> static ValueType CoerceValue(const void* Addr, int16 SizeAndType);
 };
 
+////////////////////////////////////////////////////////////////////////////////
+template <typename ValueType>
+IAnalyzer::FEventFieldHandle IAnalyzer::FEventTypeInfo::GetFieldHandle(const ANSICHAR* FieldName) const
+{
+	int16 SizeAndType;
+	FEventFieldHandle Handle = GetFieldHandleImpl(FieldName, SizeAndType);
+	if (std::is_floating_point<ValueType>::value)
+	{
+		checkf((SizeAndType < 0), TEXT("Field is not a float-type field"));
+	}
+	else
+	{
+		checkf(SizeAndType >= sizeof(ValueType), TEXT("Field is to small to read as hinted type ValueType"));
+	}
+	return Handle;
+}
+
 ////////////////////////////////////////////////////////////////////////////////
 template <typename ValueType>
 ValueType IAnalyzer::CoerceValue(const void* Addr, int16 SizeAndType)
@@ -280,6 +316,14 @@ ValueType IAnalyzer::FEventData::GetValue(const ANSICHAR* FieldName, ValueType D
 	return Default;
 }
 
+////////////////////////////////////////////////////////////////////////////////
+template <typename ValueType>
+ValueType IAnalyzer::FEventData::GetValue(FEventFieldHandle FieldHandle) const
+{
+	const uint8* EventDataPtr = *(const uint8**)this;
+	return *(ValueType*)(EventDataPtr + FieldHandle.Detail);
+}
+
 ////////////////////////////////////////////////////////////////////////////////
 template <typename ValueType>
 const IAnalyzer::TArrayReader<ValueType>& IAnalyzer::FEventData::GetArray(const ANSICHAR* FieldName) const
