Indicate when it is okay for analysis to derive the starting log serial. This mitigates against cases where the data sent in one update of the run time has fragmented delivery when analysed (e.g. due to TCP overhead or read buffer sizes). Fragmentation results in analysis making decisions based on a partial view of the world.

#rnx

diff --git a/Engine/Source/Developer/TraceAnalysis/Private/Analysis/Engine.cpp b/Engine/Source/Developer/TraceAnalysis/Private/Analysis/Engine.cpp
--- a/Engine/Source/Developer/TraceAnalysis/Private/Analysis/Engine.cpp
+++ b/Engine/Source/Developer/TraceAnalysis/Private/Analysis/Engine.cpp
@@ -716,6 +716,7 @@ enum ERouteId : uint16
 {
 	RouteId_NewEvent,
 	RouteId_NewTrace,
+	RouteId_SerialSync,
 	RouteId_Timing,
 	RouteId_ThreadTiming,
 	RouteId_ThreadInfo,
@@ -919,6 +920,7 @@ void FAnalysisEngine::OnAnalysisBegin(const FOnAnalysisContext& Context)
 {
 	auto& Builder = Context.InterfaceBuilder;
 	Builder.RouteEvent(RouteId_NewTrace,		"$Trace", "NewTrace");
+	Builder.RouteEvent(RouteId_SerialSync,		"$Trace", "SerialSync");
 	Builder.RouteEvent(RouteId_NewEvent,		"$Trace", "NewEvent");
 	Builder.RouteEvent(RouteId_Timing,			"$Trace", "Timing");
 	Builder.RouteEvent(RouteId_ThreadTiming,	"$Trace", "ThreadTiming");
@@ -940,6 +942,7 @@ bool FAnalysisEngine::OnEvent(uint16 RouteId, EStyle Style, const FOnEventContex
 	switch (RouteId)
 	{
 	case RouteId_NewTrace:			OnNewTrace(Context);			break;
+	case RouteId_SerialSync:		OnSerialSync(Context);			break;
 	case RouteId_Timing:			OnTiming(Context);				break;
 	case RouteId_ThreadTiming:		OnThreadTiming(Context);		break;
 	case RouteId_ThreadInfo:		OnThreadInfoInternal(Context);	break;
@@ -961,13 +964,27 @@ void FAnalysisEngine::OnNewTrace(const FOnEventContext& Context)
 	uint32 Hint = EventData.GetValue<uint32>("Serial");
 	Hint -= (Serial.Mask + 1) >> 1;
 	Hint &= Serial.Mask;
-	Serial.Value = Hint|0x80000000;
+	Serial.Value = Hint;
+	Serial.Value |= 0xc0000000;
+
+	// Later traces will have an explicit "SerialSync" trace event to indicate
+	// when there is enough data to establish the correct log serial
+	if ((EventData.GetValue<uint8>("FeatureSet") & 1) == 0)
+	{
+		OnSerialSync(Context);
+	}
 
 	UserUidBias = EventData.GetValue<uint32>("UserUidBias", uint32(UE::Trace::Protocol3::EKnownEventUids::User));
 
 	OnTiming(Context);
 }
 
+////////////////////////////////////////////////////////////////////////////////
+void FAnalysisEngine::OnSerialSync(const FOnEventContext& Context)
+{
+	Serial.Value &= ~0x40000000;
+}
+
 ////////////////////////////////////////////////////////////////////////////////
 void FAnalysisEngine::OnTiming(const FOnEventContext& Context)
 {
@@ -1535,7 +1552,7 @@ bool FAnalysisEngine::OnDataProtocol2()
 
 		// If the current serial has its MSB set we're currently in a mode trying
 		// to derive the best starting serial.
-		if (int32(Serial.Value) < 0)
+		if (int32(Serial.Value) < 0 && int32(Serial.Value) < int32(0xc0000000))
 		{
 			Serial.Value = (MinLogSerial & Serial.Mask);
 			continue;
diff --git a/Engine/Source/Developer/TraceAnalysis/Private/Analysis/Engine.h b/Engine/Source/Developer/TraceAnalysis/Private/Analysis/Engine.h
--- a/Engine/Source/Developer/TraceAnalysis/Private/Analysis/Engine.h
+++ b/Engine/Source/Developer/TraceAnalysis/Private/Analysis/Engine.h
@@ -90,6 +90,7 @@ private:
 	virtual bool		OnEvent(uint16 RouteId, EStyle Style, const FOnEventContext& Context) override;
 	virtual void		OnAnalysisBegin(const FOnAnalysisContext& Context) override;
 	void				OnNewTrace(const FOnEventContext& Context);
+	void				OnSerialSync(const FOnEventContext& Context);
 	void				OnThreadTiming(const FOnEventContext& Context);
 	void				OnThreadInfoInternal(const FOnEventContext& Context);
 	void				OnThreadGroupBegin(const FOnEventContext& Context);
diff --git a/Engine/Source/Runtime/TraceLog/Private/Trace/Writer.cpp b/Engine/Source/Runtime/TraceLog/Private/Trace/Writer.cpp
--- a/Engine/Source/Runtime/TraceLog/Private/Trace/Writer.cpp
+++ b/Engine/Source/Runtime/TraceLog/Private/Trace/Writer.cpp
@@ -55,6 +55,10 @@ UE_TRACE_EVENT_BEGIN($Trace, NewTrace, NoSync)
 	UE_TRACE_EVENT_FIELD(uint16, UserUidBias)
 	UE_TRACE_EVENT_FIELD(uint16, Endian)
 	UE_TRACE_EVENT_FIELD(uint8, PointerSize)
+	UE_TRACE_EVENT_FIELD(uint8, FeatureSet)
+UE_TRACE_EVENT_END()
+
+UE_TRACE_EVENT_BEGIN($Trace, SerialSync, NoSync)
 UE_TRACE_EVENT_END()
 
 
@@ -199,6 +203,7 @@ void Writer_MemoryFree(void* Address, uint32 Size)
 
 ////////////////////////////////////////////////////////////////////////////////
 static UPTRINT					GDataHandle;		// = 0
+static bool						GSerialSyncPending;	// = false
 UPTRINT							GPendingDataHandle;	// = 0
 
 ////////////////////////////////////////////////////////////////////////////////
@@ -333,17 +338,29 @@ static void Writer_DescribeAnnounce()
 static void Writer_LogHeader()
 {
 	UE_TRACE_LOG($Trace, NewTrace, TraceLogChannel)
-		<< NewTrace.Serial(uint32(GLogSerial)) // should really atomic-load-relaxed here...
+		<< NewTrace.Serial(AtomicLoadRelaxed(&GLogSerial))
 		<< NewTrace.UserUidBias(EKnownEventUids::User)
 		<< NewTrace.Endian(0x524d)
 		<< NewTrace.PointerSize(sizeof(void*))
 		<< NewTrace.StartCycle(GStartCycle)
-		<< NewTrace.CycleFrequency(TimeGetFrequency());
+		<< NewTrace.CycleFrequency(TimeGetFrequency())
+		<< NewTrace.FeatureSet(1);
 }
 
 ////////////////////////////////////////////////////////////////////////////////
 static bool Writer_UpdateConnection()
 {
+	if (GSerialSyncPending)
+	{
+		GSerialSyncPending = false;
+
+		// When analysis receives the "serial sync" event the starting serial
+		// can be established from preceeding synced events.
+		TWriteBufferRedirect<32> SideBuffer;
+		UE_TRACE_LOG($Trace, SerialSync, TraceLogChannel);
+		Writer_SendData(SideBuffer.GetData(), SideBuffer.GetSize());
+	}
+
 	if (!GPendingDataHandle)
 	{
 		return false;
@@ -366,6 +383,7 @@ static bool Writer_UpdateConnection()
 	{
 		IoClose(GPendingDataHandle);
 		GPendingDataHandle = 0;
+		GSerialSyncPending = false;
 		return false;
 	}
 
@@ -418,6 +436,7 @@ static bool Writer_UpdateConnection()
 
 	Writer_CacheOnConnect();
 
+	GSerialSyncPending = true;
 	return true;
 }
 
