Fixed PVS-Studio analysis warning.

#rnx

diff --git a/Engine/Source/Developer/TraceAnalysis/Private/Analysis/Engine.cpp b/Engine/Source/Developer/TraceAnalysis/Private/Analysis/Engine.cpp
--- a/Engine/Source/Developer/TraceAnalysis/Private/Analysis/Engine.cpp
+++ b/Engine/Source/Developer/TraceAnalysis/Private/Analysis/Engine.cpp
@@ -1552,7 +1552,7 @@ bool FAnalysisEngine::OnDataProtocol2()
 
 		// If the current serial has its MSB set we're currently in a mode trying
 		// to derive the best starting serial.
-		if (int32(Serial.Value) < 0 && int32(Serial.Value) < int32(0xc0000000))
+		if (int32(Serial.Value) < int32(0xc0000000))
 		{
 			Serial.Value = (MinLogSerial & Serial.Mask);
 			continue;
