Added access to Trace.Stop() via the "command-and-control" server.

#rnx

diff --git a/Engine/Source/Developer/TraceAnalysis/Private/ControlClient.cpp b/Engine/Source/Developer/TraceAnalysis/Private/ControlClient.cpp
--- a/Engine/Source/Developer/TraceAnalysis/Private/ControlClient.cpp
+++ b/Engine/Source/Developer/TraceAnalysis/Private/ControlClient.cpp
@@ -102,25 +102,21 @@ bool FControlClient::IsConnected() const
 ////////////////////////////////////////////////////////////////////////////////
 void FControlClient::SendSendTo(const TCHAR* Host)
 {
-    if (!IsConnected())
-    {
-        return;
-    }
-
     FormatAndSend(TEXT("SendTo %s"), Host);
 }
 
 ////////////////////////////////////////////////////////////////////////////////
 void FControlClient::SendWriteTo(const TCHAR* Path)
 {
-    if (!IsConnected())
-    {
-        return;
-    }
-
     FormatAndSend(TEXT("WriteTo %s"), Path);
 }
 
+////////////////////////////////////////////////////////////////////////////////
+void FControlClient::SendStop()
+{
+    FormatAndSend(TEXT("Stop"));
+}
+
 ////////////////////////////////////////////////////////////////////////////////
 void FControlClient::SendToggleChannel(const TCHAR* Channels, bool bState /*= true*/)
 {
diff --git a/Engine/Source/Developer/TraceAnalysis/Public/Trace/ControlClient.h b/Engine/Source/Developer/TraceAnalysis/Public/Trace/ControlClient.h
--- a/Engine/Source/Developer/TraceAnalysis/Public/Trace/ControlClient.h
+++ b/Engine/Source/Developer/TraceAnalysis/Public/Trace/ControlClient.h
@@ -44,6 +44,9 @@ public:
 	 * @param Path File system path that the runtime should write events to. */
     void SendWriteTo(const TCHAR* Path);
 
+	/** Tells the runtime to stop tracing. */
+    void SendStop();
+
 	/**
 	 * Toggles one or more channels on and off. A channel or a comma separated list
 	 * of channels can be controlled.
diff --git a/Engine/Source/Runtime/TraceLog/Private/Trace/Control.cpp b/Engine/Source/Runtime/TraceLog/Private/Trace/Control.cpp
--- a/Engine/Source/Runtime/TraceLog/Private/Trace/Control.cpp
+++ b/Engine/Source/Runtime/TraceLog/Private/Trace/Control.cpp
@@ -16,6 +16,7 @@ namespace Private {
 ////////////////////////////////////////////////////////////////////////////////
 bool	Writer_SendTo(const ANSICHAR*, uint32);
 bool	Writer_WriteTo(const ANSICHAR*);
+bool	Writer_Stop();
 
 
 
@@ -301,6 +302,13 @@ void Writer_InitializeControl()
 		}
 	);
 
+	Writer_ControlAddCommand("Stop", nullptr,
+		[] (void*, uint32 ArgC, ANSICHAR const* const* ArgV)
+		{
+			Writer_Stop();
+		}
+	);
+
 	Writer_ControlAddCommand("ToggleChannels", nullptr, 
 		[] (void*, uint32 ArgC, ANSICHAR const* const* ArgV) 
 		{
