Corrected a comment's typo.

#rnx

diff --git a/Engine/Source/Developer/TraceAnalysis/Public/Trace/Analyzer.h b/Engine/Source/Developer/TraceAnalysis/Public/Trace/Analyzer.h
--- a/Engine/Source/Developer/TraceAnalysis/Public/Trace/Analyzer.h
+++ b/Engine/Source/Developer/TraceAnalysis/Public/Trace/Analyzer.h
@@ -71,7 +71,7 @@ public:
 
 	struct TRACEANALYSIS_API FEventTypeInfo
 	{
-		/** Each event is assigned a unique ID when logged. Not that this is not
+		/** Each event is assigned a unique ID when logged. Note that this is not
 		 * guaranteed to be the same for the same event from one trace to the next. */
 		uint32 GetId() const;
 
