Add the internal thread into the transport handler's list straight away and always keep it there. Saves the round trip of always popping it in and out to ensure it gets processed first.

#rnx

diff --git a/Engine/Source/Developer/TraceAnalysis/Private/Analysis/Transport/TidPacketTransport.cpp b/Engine/Source/Developer/TraceAnalysis/Private/Analysis/Transport/TidPacketTransport.cpp
--- a/Engine/Source/Developer/TraceAnalysis/Private/Analysis/Transport/TidPacketTransport.cpp
+++ b/Engine/Source/Developer/TraceAnalysis/Private/Analysis/Transport/TidPacketTransport.cpp
@@ -3,7 +3,6 @@
 #include "TidPacketTransport.h"
 #include "Algo/BinarySearch.h"
 #include "HAL/UnrealMemory.h"
-#include "Trace/Detail/Transport.h"
 
 ////////////////////////////////////////////////////////////////////////////////
 namespace UE {
@@ -77,15 +76,6 @@ FTidPacketTransport::FThreadStream& FTidPacketTransport::FindOrAddThread(uint32
 
 	FThreadStream Thread;
 	Thread.ThreadId = ThreadId;
-
-	// Internal events are sent over tid 0 and they should be processed first. We
-	// don't care about the order otherwise.
-	if (Thread.ThreadId == UE::Trace::ETransportTid::Internal)
-	{
-		Threads.Insert(Thread, 0);
-		return Threads[0];
-	}
-
 	Threads.Add(Thread);
 	return Threads[ThreadCount];
 }
@@ -97,7 +87,7 @@ void FTidPacketTransport::Update()
 
 	Threads.RemoveAll([] (const FThreadStream& Thread)
 	{
-		return Thread.Buffer.IsEmpty();
+		return (Thread.ThreadId <= ETransportTid::Internal) ? false : Thread.Buffer.IsEmpty();
 	});
 }
 
diff --git a/Engine/Source/Developer/TraceAnalysis/Private/Analysis/Transport/TidPacketTransport.h b/Engine/Source/Developer/TraceAnalysis/Private/Analysis/Transport/TidPacketTransport.h
--- a/Engine/Source/Developer/TraceAnalysis/Private/Analysis/Transport/TidPacketTransport.h
+++ b/Engine/Source/Developer/TraceAnalysis/Private/Analysis/Transport/TidPacketTransport.h
@@ -5,6 +5,7 @@
 #include "Transport.h"
 #include "Analysis/StreamReader.h"
 #include "Containers/Array.h"
+#include "Trace/Detail/Transport.h"
 
 namespace UE {
 namespace Trace {
@@ -32,7 +33,9 @@ private:
 
 	bool					ReadPacket();
 	FThreadStream&			FindOrAddThread(uint32 ThreadId);
-	TArray<FThreadStream>	Threads;
+	TArray<FThreadStream>	Threads = {
+								{ {}, ETransportTid::Internal },
+							};
 };
 
 } // namespace Trace
