Added an EventTypeInfo.GetSize() to get the unadorned size of an trace event.

#rnx

diff --git a/Engine/Source/Developer/TraceAnalysis/Private/Analysis/Engine.cpp b/Engine/Source/Developer/TraceAnalysis/Private/Analysis/Engine.cpp
--- a/Engine/Source/Developer/TraceAnalysis/Private/Analysis/Engine.cpp
+++ b/Engine/Source/Developer/TraceAnalysis/Private/Analysis/Engine.cpp
@@ -417,6 +417,13 @@ const ANSICHAR* IAnalyzer::FEventTypeInfo::GetLoggerName() const
 	return (const ANSICHAR*)Inner + Inner->LoggerNameOffset;
 }
 
+////////////////////////////////////////////////////////////////////////////////
+uint32 IAnalyzer::FEventTypeInfo::GetSize() const
+{
+	const auto* Inner = (const FAnalysisEngine::FDispatch*)this;
+	return Inner->EventSize;
+}
+
 ////////////////////////////////////////////////////////////////////////////////
 uint32 IAnalyzer::FEventTypeInfo::GetFieldCount() const
 {
diff --git a/Engine/Source/Developer/TraceAnalysis/Public/Trace/Analyzer.h b/Engine/Source/Developer/TraceAnalysis/Public/Trace/Analyzer.h
--- a/Engine/Source/Developer/TraceAnalysis/Public/Trace/Analyzer.h
+++ b/Engine/Source/Developer/TraceAnalysis/Public/Trace/Analyzer.h
@@ -81,6 +81,9 @@ public:
 		/** Returns the logger name the event is associated with. */
 		const ANSICHAR* GetLoggerName() const;
 
+		/** Returns the base size of the event. */
+		uint32 GetSize() const;
+
 		/** The number of member fields this event has. */
 		uint32 GetFieldCount() const;
 
