Added a EventData.GetSize() method to report the size in bytes of a traced event to the curious.

#rnx

diff --git a/Engine/Source/Developer/TraceAnalysis/Private/Analysis/Engine.cpp b/Engine/Source/Developer/TraceAnalysis/Private/Analysis/Engine.cpp
--- a/Engine/Source/Developer/TraceAnalysis/Private/Analysis/Engine.cpp
+++ b/Engine/Source/Developer/TraceAnalysis/Private/Analysis/Engine.cpp
@@ -691,6 +691,18 @@ bool IAnalyzer::FEventData::GetString(const ANSICHAR* FieldName, FString& Out) c
 	return false;
 }
 
+////////////////////////////////////////////////////////////////////////////////
+uint32 IAnalyzer::FEventData::GetSize() const
+{
+	const auto* Info = (const FAnalysisEngine::FEventDataInfo*)this;
+	uint32 Size = Info->Size;
+	for (const FAuxData& Data : *(Info->AuxCollector))
+	{
+		Size += Data.DataSize;
+	}
+	return Size;
+}
+
 ////////////////////////////////////////////////////////////////////////////////
 void IAnalyzer::FEventData::SerializeToCbor(TArray<uint8>& Out) const
 {
diff --git a/Engine/Source/Developer/TraceAnalysis/Public/Trace/Analyzer.h b/Engine/Source/Developer/TraceAnalysis/Public/Trace/Analyzer.h
--- a/Engine/Source/Developer/TraceAnalysis/Public/Trace/Analyzer.h
+++ b/Engine/Source/Developer/TraceAnalysis/Public/Trace/Analyzer.h
@@ -156,6 +156,9 @@ public:
 		bool GetString(const ANSICHAR* FieldName, FStringView& Out) const;
 		bool GetString(const ANSICHAR* FieldName, FString& Out) const;
 
+		/** The size of the event in uncompressed bytes excluding the header */
+		uint32 GetSize() const;
+
 		/** Serializes the event to Cbor object.
 		 * @param Recipient of the Cbor serialization. Data is appeneded to Out. */
 		void SerializeToCbor(TArray<uint8>& Out) const;
