EventData.GetString() was incorrectly returning false if the requested field did exist but was empty.

#rnx

diff --git a/Engine/Source/Developer/TraceAnalysis/Private/Analysis/Engine.cpp b/Engine/Source/Developer/TraceAnalysis/Private/Analysis/Engine.cpp
--- a/Engine/Source/Developer/TraceAnalysis/Private/Analysis/Engine.cpp
+++ b/Engine/Source/Developer/TraceAnalysis/Private/Analysis/Engine.cpp
@@ -624,7 +624,8 @@ bool IAnalyzer::FEventData::GetString(const ANSICHAR* FieldName, FAnsiStringView
 		return true;
 	}
 
-	return false;
+	Out = FAnsiStringView();
+	return true;
 }
 
 ////////////////////////////////////////////////////////////////////////////////
@@ -651,7 +652,8 @@ bool IAnalyzer::FEventData::GetString(const ANSICHAR* FieldName, FStringView& Ou
 		return true;
 	}
 
-	return false;
+	Out = FStringView();
+	return true;
 }
 
 ////////////////////////////////////////////////////////////////////////////////
@@ -666,15 +668,16 @@ bool IAnalyzer::FEventData::GetString(const ANSICHAR* FieldName, FString& Out) c
 		return false;
 	}
 
-	const FAuxData* Data = Info->GetAuxData(Index);
-	if (Data == nullptr)
-	{
-		return false;
-	}
-
 	const auto& Field = Dispatch.Fields[Index];
 	if (Field.Class == UE::Trace::Protocol0::Field_String)
 	{
+		const FAuxData* Data = Info->GetAuxData(Index);
+		if (Data == nullptr)
+		{
+			Out = FString();
+			return true;
+		}
+
 		if (Field.SizeAndType == sizeof(ANSICHAR))
 		{
 			Out = FString(Data->DataSize, (const ANSICHAR*)(Data->Data));
