Guard against some legacy paths taken when analyzing old traces as they set the auxilary data collector object to null.

#rnx

diff --git a/Engine/Source/Developer/TraceAnalysis/Private/Analysis/Engine.cpp b/Engine/Source/Developer/TraceAnalysis/Private/Analysis/Engine.cpp
--- a/Engine/Source/Developer/TraceAnalysis/Private/Analysis/Engine.cpp
+++ b/Engine/Source/Developer/TraceAnalysis/Private/Analysis/Engine.cpp
@@ -711,9 +711,12 @@ void IAnalyzer::FEventData::SerializeToCbor(TArray<uint8>& Out) const
 {
 	const auto* Info = (const FAnalysisEngine::FEventDataInfo*)this;
 	uint32 Size = Info->Size;
-	for (FAuxData& Data : *(Info->AuxCollector))
+	if (Info->AuxCollector != nullptr)
 	{
-		Size += Data.DataSize;
+		for (FAuxData& Data : *(Info->AuxCollector))
+		{
+			Size += Data.DataSize;
+		}
 	}
 	SerializeToCborImpl(Out, *this, Size);
 }
