Use uint32 to store the UID instead of uint16. If anything the code generation is a little better as there are no prefixes.

#rnx

diff --git a/Engine/Source/Developer/TraceAnalysis/Private/Analysis/Engine.cpp b/Engine/Source/Developer/TraceAnalysis/Private/Analysis/Engine.cpp
--- a/Engine/Source/Developer/TraceAnalysis/Private/Analysis/Engine.cpp
+++ b/Engine/Source/Developer/TraceAnalysis/Private/Analysis/Engine.cpp
@@ -1426,8 +1426,8 @@ bool FAnalysisEngine::OnDataProtocol0()
 			break;
 		}
 
-		uint16 Uid = uint16(Header->Uid) & uint16(Protocol0::EKnownEventUids::UidMask);
-		if (Uid >= Dispatches.Num())
+		uint32 Uid = uint32(Header->Uid) & uint32(Protocol0::EKnownEventUids::UidMask);
+		if (Uid >= uint32(Dispatches.Num()))
 		{
 			return false;
 		}
@@ -1610,8 +1610,8 @@ int32 FAnalysisEngine::OnDataProtocol2(FStreamReader& Reader, FThreads::FInfo& T
 			break;
 		}
 
-		uint16 Uid = Header->Uid & uint16(Protocol2::EKnownEventUids::UidMask);
-		if (Uid >= Dispatches.Num())
+		uint32 Uid = uint32(Header->Uid) & uint32(Protocol2::EKnownEventUids::UidMask);
+		if (Uid >= uint32(Dispatches.Num()))
 		{
 			// We don't know about this event yet
 			break;
