Fixed typo in a comment

#rnx

diff --git a/Engine/Source/Developer/TraceAnalysis/Private/Analysis/Engine.cpp b/Engine/Source/Developer/TraceAnalysis/Private/Analysis/Engine.cpp
--- a/Engine/Source/Developer/TraceAnalysis/Private/Analysis/Engine.cpp
+++ b/Engine/Source/Developer/TraceAnalysis/Private/Analysis/Engine.cpp
@@ -985,7 +985,7 @@ void FAnalysisEngine::OnNewTrace(const FOnEventContext& Context)
 	const FEventData& EventData = Context.EventData;
 
 	// "Serial" will tell us approximately where we've started in the log serial
-	// range. We'll bias is by half so we won't accept any serialised events and
+	// range. We'll bias it by half so we won't accept any serialised events and
 	// mark the MSB to indicate that the current serial should be corrected.
 	uint32 Hint = EventData.GetValue<uint32>("Serial");
 	Hint -= (Serial.Mask + 1) >> 1;
