Guard against reading too much of an event's header from the trace data stream

#rnx

diff --git a/Engine/Source/Developer/TraceAnalysis/Private/Analysis/Engine.cpp b/Engine/Source/Developer/TraceAnalysis/Private/Analysis/Engine.cpp
--- a/Engine/Source/Developer/TraceAnalysis/Private/Analysis/Engine.cpp
+++ b/Engine/Source/Developer/TraceAnalysis/Private/Analysis/Engine.cpp
@@ -1919,7 +1919,12 @@ int32 FAnalysisEngine::OnDataProtocol4Impl(
 	// Make sure we consume events in the correct order
 	if ((Dispatch->Flags & FDispatch::Flag_NoSync) == 0)
 	{
-		const auto* HeaderSync = (UE::Trace::Protocol4::FEventHeaderSync*)Header;
+		if (Reader.GetPointer<FEventHeaderSync>() == nullptr)
+		{
+			return 1;
+		}
+		
+		const auto* HeaderSync = (FEventHeaderSync*)Header;
 		uint32 EventSerial = HeaderSync->SerialLow|(uint32(HeaderSync->SerialHigh) << 16);
 		if (EventSerial != (Serial.Value & Serial.Mask))
 		{
