Extend lifetime of thread info objects. Previously there was a slim chance that they would get moved while someone else was holding on to a pointer to them.

#rnx

diff --git a/Engine/Source/Developer/TraceAnalysis/Private/Analysis/Engine.cpp b/Engine/Source/Developer/TraceAnalysis/Private/Analysis/Engine.cpp
--- a/Engine/Source/Developer/TraceAnalysis/Private/Analysis/Engine.cpp
+++ b/Engine/Source/Developer/TraceAnalysis/Private/Analysis/Engine.cpp
@@ -64,6 +64,15 @@ FThreads::FThreads()
 	Infos.Reserve(64);
 }
 
+////////////////////////////////////////////////////////////////////////////////
+FThreads::~FThreads()
+{
+	for (FInfo* Info : Infos)
+	{
+		delete Info;
+	}
+}
+
 ////////////////////////////////////////////////////////////////////////////////
 FThreads::FInfo* FThreads::GetInfo()
 {
@@ -75,29 +84,19 @@ FThreads::FInfo* FThreads::GetInfo(uint32 ThreadId)
 {
 	LastGetInfoId = ThreadId;
 
-	// Events marked as important are associated internally to a pseudo thread
-	// with id 0 [A]. New-thread events are considered important and thus appear
-	// as if they were traced on thread 0. They will add thread N which will
-	// modify the Infos array. If someone was holding a A& = Infos[0] when
-	// Infos[N] was added then A can become stale. To combat this edge-case
-	// we'll thread 0 as a special case so it doesn't invalidate itself with its
-	// own events.
-	if (ThreadId == 0)
-	{
-		return &ImportantInfo;
-	}
-
 	if (ThreadId >= uint32(Infos.Num()))
 	{
-		Infos.SetNum(ThreadId);
+		Infos.SetNumZeroed(ThreadId + 1);
 	}
 
-	FInfo& Info = Infos[ThreadId - 1];
-	if (Info.ThreadId == ~0u)
+	FInfo* Info = Infos[ThreadId];
+	if (Info == nullptr)
 	{
-		Info.ThreadId = ThreadId;
+		Info = new FInfo();
+		Info->ThreadId = ThreadId;
+		Infos[ThreadId] = Info;
 	}
-	return &Info;
+	return Info;
 }
 
 ////////////////////////////////////////////////////////////////////////////////
diff --git a/Engine/Source/Developer/TraceAnalysis/Private/Analysis/Engine.h b/Engine/Source/Developer/TraceAnalysis/Private/Analysis/Engine.h
--- a/Engine/Source/Developer/TraceAnalysis/Private/Analysis/Engine.h
+++ b/Engine/Source/Developer/TraceAnalysis/Private/Analysis/Engine.h
@@ -27,6 +27,7 @@ public:
 	};
 
 						FThreads();
+						~FThreads();
 	FInfo*				GetInfo();
 	FInfo*				GetInfo(uint32 ThreadId);
 	void				SetGroupName(const ANSICHAR* InGroupName, uint32 Length);
@@ -34,8 +35,7 @@ public:
 
 private:
 	uint32				LastGetInfoId = ~0u;
-	FInfo				ImportantInfo;
-	TArray<FInfo>		Infos;
+	TArray<FInfo*>		Infos;
 	TArray<uint8>		GroupName;
 };
 
