Use Trace's packet structs on the analysis side instead of duplicating them.
#rnx

diff --git a/Engine/Source/Developer/TraceAnalysis/Private/Analysis/Transport/TidPacketTransport.cpp b/Engine/Source/Developer/TraceAnalysis/Private/Analysis/Transport/TidPacketTransport.cpp
--- a/Engine/Source/Developer/TraceAnalysis/Private/Analysis/Transport/TidPacketTransport.cpp
+++ b/Engine/Source/Developer/TraceAnalysis/Private/Analysis/Transport/TidPacketTransport.cpp
@@ -23,13 +23,9 @@ namespace Trace {
 ////////////////////////////////////////////////////////////////////////////////
 bool FTidPacketTransport::ReadPacket()
 {
-	struct FPacketBase
-	{
-		uint16	PacketSize;
-		uint16	ThreadId;
-	};
+	using namespace UE::Trace::Private;
 
-	const auto* PacketBase = GetPointer<FPacketBase>();
+	const auto* PacketBase = GetPointer<FTidPacketBase>();
 	if (PacketBase == nullptr)
 	{
 		return false;
@@ -45,14 +41,15 @@ bool FTidPacketTransport::ReadPacket()
 	uint32 ThreadId = PacketBase->ThreadId & ~0x8000;
 	FThreadStream& Thread = FindOrAddThread(ThreadId);
 
-	uint32 DataSize = PacketBase->PacketSize - sizeof(FPacketBase);
+	uint32 DataSize = PacketBase->PacketSize - sizeof(FTidPacketBase);
 	if (PacketBase->ThreadId != ThreadId)
 	{
-		uint16* DecodedSize = (uint16*)(PacketBase + 1);
-		uint8* Dest = Thread.Buffer.Append(*DecodedSize);
-		DataSize -= sizeof(*DecodedSize);
-		int32 ResultSize = UE::Trace::Private::Decode(DecodedSize + 1, DataSize, Dest, *DecodedSize);
-		check(int32(*DecodedSize) == ResultSize);
+		const auto* Packet = (const FTidPacketEncoded*)PacketBase;
+		uint16 DecodedSize = Packet->DecodedSize;
+		uint8* Dest = Thread.Buffer.Append(DecodedSize);
+		DataSize -= sizeof(DecodedSize);
+		int32 ResultSize = UE::Trace::Private::Decode(Packet->Data, DataSize, Dest, DecodedSize);
+		check(int32(DecodedSize) == ResultSize);
 	}
 	else
 	{
