The transport layer can now choose to discard packets for a particular thread. Currently none are discarded but the functionality will be used in late-connect scenarios where the first packet for a thread is partial (e.g. it starts in the middle of auxilary trace data). There is no mechanism in place to recover reading events in this case. It is however a rare occurence thus cheaper, easier, and safer to discard the packet than support recovery.

#rnx

diff --git a/Engine/Source/Developer/TraceAnalysis/Private/Analysis/Transport/TidPacketTransport.cpp b/Engine/Source/Developer/TraceAnalysis/Private/Analysis/Transport/TidPacketTransport.cpp
--- a/Engine/Source/Developer/TraceAnalysis/Private/Analysis/Transport/TidPacketTransport.cpp
+++ b/Engine/Source/Developer/TraceAnalysis/Private/Analysis/Transport/TidPacketTransport.cpp
@@ -39,42 +39,53 @@ bool FTidPacketTransport::ReadPacket()
 	FTransport::Advance(PacketBase->PacketSize);
 
 	uint32 ThreadId = PacketBase->ThreadId & FTidPacketBase::ThreadIdMask;
-	FThreadStream& Thread = FindOrAddThread(ThreadId);
+	FThreadStream* Thread = FindOrAddThread(ThreadId, true);
+	if (Thread == nullptr)
+	{
+		return true;
+	}
 
 	uint32 DataSize = PacketBase->PacketSize - sizeof(FTidPacketBase);
 	if (PacketBase->ThreadId != ThreadId)
 	{
 		const auto* Packet = (const FTidPacketEncoded*)PacketBase;
 		uint16 DecodedSize = Packet->DecodedSize;
-		uint8* Dest = Thread.Buffer.Append(DecodedSize);
+		uint8* Dest = Thread->Buffer.Append(DecodedSize);
 		DataSize -= sizeof(DecodedSize);
 		int32 ResultSize = UE::Trace::Private::Decode(Packet->Data, DataSize, Dest, DecodedSize);
 		check(int32(DecodedSize) == ResultSize);
 	}
 	else
 	{
-		Thread.Buffer.Append((uint8*)(PacketBase + 1), DataSize);
+		Thread->Buffer.Append((uint8*)(PacketBase + 1), DataSize);
 	}
 
 	return true;
 }
 
 ////////////////////////////////////////////////////////////////////////////////
-FTidPacketTransport::FThreadStream& FTidPacketTransport::FindOrAddThread(uint32 ThreadId)
+FTidPacketTransport::FThreadStream* FTidPacketTransport::FindOrAddThread(
+	uint32	ThreadId,
+	bool	bAddIfNotFound)
 {
 	uint32 ThreadCount = Threads.Num();
 	for (uint32 i = 0; i < ThreadCount; ++i)
 	{
 		if (Threads[i].ThreadId == ThreadId)
 		{
-			return Threads[i];
+			return &(Threads[i]);
 		}
 	}
+	
+	if (!bAddIfNotFound)
+	{
+		return nullptr;
+	}
 
 	FThreadStream Thread;
 	Thread.ThreadId = ThreadId;
 	Threads.Add(Thread);
-	return Threads[ThreadCount];
+	return &(Threads[ThreadCount]);
 }
 
 ////////////////////////////////////////////////////////////////////////////////
diff --git a/Engine/Source/Developer/TraceAnalysis/Private/Analysis/Transport/TidPacketTransport.h b/Engine/Source/Developer/TraceAnalysis/Private/Analysis/Transport/TidPacketTransport.h
--- a/Engine/Source/Developer/TraceAnalysis/Private/Analysis/Transport/TidPacketTransport.h
+++ b/Engine/Source/Developer/TraceAnalysis/Private/Analysis/Transport/TidPacketTransport.h
@@ -32,7 +32,7 @@ private:
 	};
 
 	bool					ReadPacket();
-	FThreadStream&			FindOrAddThread(uint32 ThreadId);
+	FThreadStream*			FindOrAddThread(uint32 ThreadId, bool bAddIfNotFound);
 	TArray<FThreadStream>	Threads = {
 								{ {}, ETransportTid::Events },
 								{ {}, ETransportTid::Importants },
