If a partial packet arrives and we know about the thread, we're expecting more data and can just append the partial packet we received. If we do not know about the thread the packet is fore, then we've late-joined and the packet we've just received is dependant on previous data that was discarded prior to connect. In that case we'll discard the packet.

#rnx

diff --git a/Engine/Source/Developer/TraceAnalysis/Private/Analysis/Transport/TidPacketTransport.cpp b/Engine/Source/Developer/TraceAnalysis/Private/Analysis/Transport/TidPacketTransport.cpp
--- a/Engine/Source/Developer/TraceAnalysis/Private/Analysis/Transport/TidPacketTransport.cpp
+++ b/Engine/Source/Developer/TraceAnalysis/Private/Analysis/Transport/TidPacketTransport.cpp
@@ -39,7 +39,8 @@ bool FTidPacketTransport::ReadPacket()
 	FTransport::Advance(PacketBase->PacketSize);
 
 	uint32 ThreadId = PacketBase->ThreadId & FTidPacketBase::ThreadIdMask;
-	FThreadStream* Thread = FindOrAddThread(ThreadId, true);
+	bool bIsPartial = !!(PacketBase->ThreadId & FTidPacketBase::PartialMarker);
+	FThreadStream* Thread = FindOrAddThread(ThreadId, !bIsPartial);
 	if (Thread == nullptr)
 	{
 		return true;
