Packets that are partial but not compressed where being incorrectly categorised as compressed. Analysis would fallover when nonsense data was sent to LZ4.

#rnx

diff --git a/Engine/Source/Developer/TraceAnalysis/Private/Analysis/Transport/TidPacketTransport.cpp b/Engine/Source/Developer/TraceAnalysis/Private/Analysis/Transport/TidPacketTransport.cpp
--- a/Engine/Source/Developer/TraceAnalysis/Private/Analysis/Transport/TidPacketTransport.cpp
+++ b/Engine/Source/Developer/TraceAnalysis/Private/Analysis/Transport/TidPacketTransport.cpp
@@ -47,7 +47,7 @@ bool FTidPacketTransport::ReadPacket()
 	}
 
 	uint32 DataSize = PacketBase->PacketSize - sizeof(FTidPacketBase);
-	if (PacketBase->ThreadId != ThreadId)
+	if (PacketBase->ThreadId & FTidPacketBase::EncodedMarker)
 	{
 		const auto* Packet = (const FTidPacketEncoded*)PacketBase;
 		uint16 DecodedSize = Packet->DecodedSize;
