Call select() on a socket before making a blocking call to recv(). Calls to recv() that are waiting for data may not return even if the socket is closed. This can result in a technical hang where Insights waits on Trace waiting on a socket that was closed. This way we only call recv() when we know there is data to read, and the select()+timeout combo correctly detects to the close().

#rnx

diff --git a/Engine/Source/Developer/TraceAnalysis/Private/Store/StoreClient.cpp b/Engine/Source/Developer/TraceAnalysis/Private/Store/StoreClient.cpp
--- a/Engine/Source/Developer/TraceAnalysis/Private/Store/StoreClient.cpp
+++ b/Engine/Source/Developer/TraceAnalysis/Private/Store/StoreClient.cpp
@@ -47,21 +47,49 @@ bool FTraceDataStream::IsOpen() const
 ////////////////////////////////////////////////////////////////////////////////
 void FTraceDataStream::Close()
 {
+	Socket.shutdown(asio::ip::tcp::socket::shutdown_receive);
 	Socket.close();
 }
 
 ////////////////////////////////////////////////////////////////////////////////
 int32 FTraceDataStream::Read(void* Dest, uint32 DestSize)
 {
-	asio::error_code ErrorCode;
-	size_t BytesRead = Socket.read_some(asio::buffer(Dest, DestSize), ErrorCode);
-	if (ErrorCode)
+	auto Handle = Socket.native_handle();
+
+	fd_set Fds;
+	FD_ZERO(&Fds);
+	FD_SET(Handle, &Fds);
+
+	timeval Timeout;
+	Timeout.tv_sec = 1;
+	Timeout.tv_usec = 0;
+
+	while (true)
 	{
-		Close();
-		return -1;
+		fd_set ReadFds = Fds;
+		int Ret = select(Handle + 1, &ReadFds, 0, 0, &Timeout);
+
+		if (Ret < 0)
+		{
+			Close();
+			return -1;
+		}
+
+		if (Ret == 0)
+		{
+			continue;
+		}
+
+		asio::error_code ErrorCode;
+		size_t BytesRead = Socket.read_some(asio::buffer(Dest, DestSize), ErrorCode);
+		if (ErrorCode)
+		{
+			Close();
+			return -1;
+		}
+
+		return int32(BytesRead);
 	}
-
-	return int32(BytesRead);
 }
 
 
