Specifying the port that the Trace Store client should connect is now optional.

#rnx

diff --git a/Engine/Source/Developer/TraceAnalysis/Private/Store/StoreClient.cpp b/Engine/Source/Developer/TraceAnalysis/Private/Store/StoreClient.cpp
--- a/Engine/Source/Developer/TraceAnalysis/Private/Store/StoreClient.cpp
+++ b/Engine/Source/Developer/TraceAnalysis/Private/Store/StoreClient.cpp
@@ -182,6 +182,8 @@ const FResponse& FStoreCborClient::GetResponse() const
 ////////////////////////////////////////////////////////////////////////////////
 bool FStoreCborClient::Connect(const TCHAR* Host, uint16 Port)
 {
+	Port = (Port == 0) ? 1988 : Port;
+
 	FTCHARToUTF8 HostUtf8(Host);
 	char PortString[8];
 	FCStringAnsi::Sprintf(PortString, "%d", Port);
diff --git a/Engine/Source/Developer/TraceAnalysis/Public/Trace/StoreClient.h b/Engine/Source/Developer/TraceAnalysis/Public/Trace/StoreClient.h
--- a/Engine/Source/Developer/TraceAnalysis/Public/Trace/StoreClient.h
+++ b/Engine/Source/Developer/TraceAnalysis/Public/Trace/StoreClient.h
@@ -39,7 +39,7 @@ public:
 	};
 	
 						~FStoreClient() = default;
-	static FStoreClient*Connect(const TCHAR* Host, uint32 Port);
+	static FStoreClient*Connect(const TCHAR* Host, uint32 Port=0);
 	void				operator delete (void* Addr);
 	bool				IsValid() const;
 	uint32				GetStoreAddress() const;
