Versioned the end points that the trace store server and client communicate over. This will be useful as I shift to running the trace store in a separate process who's lifetime I've less control over.

diff --git a/Engine/Source/Developer/TraceAnalysis/Private/Store/AsioStoreCborServer.cpp b/Engine/Source/Developer/TraceAnalysis/Private/Store/AsioStoreCborServer.cpp
--- a/Engine/Source/Developer/TraceAnalysis/Private/Store/AsioStoreCborServer.cpp
+++ b/Engine/Source/Developer/TraceAnalysis/Private/Store/AsioStoreCborServer.cpp
@@ -23,7 +23,6 @@ public:
 
 protected:
 	void					OnPayload();
-	void					OnConnect();
 	void					OnSessionCount();
 	void					OnSessionInfo();
 	void					OnStatus();
@@ -74,14 +73,6 @@ void FAsioStoreCborPeer::Close()
 	Socket.Close();
 }
 
-////////////////////////////////////////////////////////////////////////////////
-void FAsioStoreCborPeer::OnConnect()
-{
-	TPayloadBuilder<> Builder(EStatusCode::Success);
-	Builder.AddInteger("version", int32(EStoreVersion::Value));
-	SendResponse(Builder.Done());
-}
-
 ////////////////////////////////////////////////////////////////////////////////
 void FAsioStoreCborPeer::OnSessionCount()
 {
@@ -237,13 +228,12 @@ void FAsioStoreCborPeer::OnPayload()
 		uint32	Hash;
 		void	(FAsioStoreCborPeer::*Func)();
 	} const DispatchTable[] = {
-		{ QuickStoreHash("connect"),		&FAsioStoreCborPeer::OnConnect },
-		{ QuickStoreHash("session/count"),	&FAsioStoreCborPeer::OnSessionCount },
-		{ QuickStoreHash("session/info"),	&FAsioStoreCborPeer::OnSessionInfo },
-		{ QuickStoreHash("status"),			&FAsioStoreCborPeer::OnStatus },
-		{ QuickStoreHash("trace/count"),	&FAsioStoreCborPeer::OnTraceCount },
-		{ QuickStoreHash("trace/info"),		&FAsioStoreCborPeer::OnTraceInfo },
-		{ QuickStoreHash("trace/read"),		&FAsioStoreCborPeer::OnTraceRead },
+		{ QuickStoreHash("v1/session/count"),	&FAsioStoreCborPeer::OnSessionCount },
+		{ QuickStoreHash("v1/session/info"),	&FAsioStoreCborPeer::OnSessionInfo },
+		{ QuickStoreHash("v1/status"),			&FAsioStoreCborPeer::OnStatus },
+		{ QuickStoreHash("v1/trace/count"),		&FAsioStoreCborPeer::OnTraceCount },
+		{ QuickStoreHash("v1/trace/info"),		&FAsioStoreCborPeer::OnTraceInfo },
+		{ QuickStoreHash("v1/trace/read"),		&FAsioStoreCborPeer::OnTraceRead },
 	};
 
 	uint32 MethodHash = QuickStoreHash(Method);
diff --git a/Engine/Source/Developer/TraceAnalysis/Private/Store/StoreClient.cpp b/Engine/Source/Developer/TraceAnalysis/Private/Store/StoreClient.cpp
--- a/Engine/Source/Developer/TraceAnalysis/Private/Store/StoreClient.cpp
+++ b/Engine/Source/Developer/TraceAnalysis/Private/Store/StoreClient.cpp
@@ -182,7 +182,7 @@ const FResponse& FStoreCborClient::GetResponse() const
 ////////////////////////////////////////////////////////////////////////////////
 bool FStoreCborClient::Connect(const TCHAR* Host, uint16 Port)
 {
-	Port = (Port == 0) ? 1988 : Port;
+	Port = (Port == 0) ? 1989 : Port;
 
 	FTCHARToUTF8 HostUtf8(Host);
 	char PortString[8];
@@ -206,10 +206,7 @@ bool FStoreCborClient::Connect(const TCHAR* Host, uint16 Port)
 		return false;
 	}
 
-	TPayloadBuilder<> Builder("connect");
-	Builder.AddInteger("version", int32(EStoreVersion::Value));
-	FPayload Payload = Builder.Done();
-	return Communicate(Payload);
+	return true;
 }
 
 ////////////////////////////////////////////////////////////////////////////////
@@ -266,7 +263,7 @@ bool FStoreCborClient::Communicate(const FPayload& Payload)
 ////////////////////////////////////////////////////////////////////////////////
 bool FStoreCborClient::GetStatus()
 {
-	TPayloadBuilder<32> Builder("status");
+	TPayloadBuilder<32> Builder("v1/status");
 	FPayload Payload = Builder.Done();
 	return Communicate(Payload);
 }
@@ -274,7 +271,7 @@ bool FStoreCborClient::GetStatus()
 ////////////////////////////////////////////////////////////////////////////////
 bool FStoreCborClient::GetTraceCount()
 {
-	TPayloadBuilder<32> Builder("trace/count");
+	TPayloadBuilder<32> Builder("v1/trace/count");
 	FPayload Payload = Builder.Done();
 	return Communicate(Payload);
 }
@@ -282,7 +279,7 @@ bool FStoreCborClient::GetTraceCount()
 ////////////////////////////////////////////////////////////////////////////////
 bool FStoreCborClient::GetTraceInfo(uint32 Index)
 {
-	TPayloadBuilder<> Builder("trace/info");
+	TPayloadBuilder<> Builder("v1/trace/info");
 	Builder.AddInteger("index", Index);
 	FPayload Payload = Builder.Done();
 	return Communicate(Payload);
@@ -291,7 +288,7 @@ bool FStoreCborClient::GetTraceInfo(uint32 Index)
 ////////////////////////////////////////////////////////////////////////////////
 bool FStoreCborClient::GetTraceInfoById(uint32 Id)
 {
-	TPayloadBuilder<> Builder("trace/info");
+	TPayloadBuilder<> Builder("v1/trace/info");
 	Builder.AddInteger("id", int32(Id));
 	FPayload Payload = Builder.Done();
 	return Communicate(Payload);
@@ -300,7 +297,7 @@ bool FStoreCborClient::GetTraceInfoById(uint32 Id)
 ////////////////////////////////////////////////////////////////////////////////
 FTraceDataStream* FStoreCborClient::ReadTrace(uint32 Id)
 {
-	TPayloadBuilder<> Builder("trace/read");
+	TPayloadBuilder<> Builder("v1/trace/read");
 	Builder.AddInteger("id", Id);
 	FPayload Payload = Builder.Done();
 	if (!Communicate(Payload))
@@ -339,7 +336,7 @@ FTraceDataStream* FStoreCborClient::ReadTrace(uint32 Id)
 ////////////////////////////////////////////////////////////////////////////////
 bool FStoreCborClient::GetSessionCount()
 {
-	TPayloadBuilder<> Builder("session/count");
+	TPayloadBuilder<> Builder("v1/session/count");
 	FPayload Payload = Builder.Done();
 	return Communicate(Payload);
 }
@@ -347,7 +344,7 @@ bool FStoreCborClient::GetSessionCount()
 ////////////////////////////////////////////////////////////////////////////////
 bool FStoreCborClient::GetSessionInfo(uint32 Index)
 {
-	TPayloadBuilder<> Builder("session/info");
+	TPayloadBuilder<> Builder("v1/session/info");
 	Builder.AddInteger("index", int32(Index));
 	FPayload Payload = Builder.Done();
 	return Communicate(Payload);
@@ -356,7 +353,7 @@ bool FStoreCborClient::GetSessionInfo(uint32 Index)
 ////////////////////////////////////////////////////////////////////////////////
 bool FStoreCborClient::GetSessionInfoById(uint32 Id)
 {
-	TPayloadBuilder<> Builder("session/info");
+	TPayloadBuilder<> Builder("v1/session/info");
 	Builder.AddInteger("id", int32(Id));
 	FPayload Payload = Builder.Done();
 	return Communicate(Payload);
@@ -365,7 +362,7 @@ bool FStoreCborClient::GetSessionInfoById(uint32 Id)
 ////////////////////////////////////////////////////////////////////////////////
 bool FStoreCborClient::GetSessionInfoByTraceId(uint32 TraceId)
 {
-	TPayloadBuilder<> Builder("session/info");
+	TPayloadBuilder<> Builder("v1/session/info");
 	Builder.AddInteger("trace_id", int32(TraceId));
 	FPayload Payload = Builder.Done();
 	return Communicate(Payload);
