Bumped a port number one step too early.

diff --git a/Engine/Source/Developer/TraceAnalysis/Private/Store/StoreClient.cpp b/Engine/Source/Developer/TraceAnalysis/Private/Store/StoreClient.cpp
--- a/Engine/Source/Developer/TraceAnalysis/Private/Store/StoreClient.cpp
+++ b/Engine/Source/Developer/TraceAnalysis/Private/Store/StoreClient.cpp
@@ -182,7 +182,7 @@ const FResponse& FStoreCborClient::GetResponse() const
 ////////////////////////////////////////////////////////////////////////////////
 bool FStoreCborClient::Connect(const TCHAR* Host, uint16 Port)
 {
-	Port = (Port == 0) ? 1989 : Port;
+	Port = (Port == 0) ? 1988 : Port;
 
 	FTCHARToUTF8 HostUtf8(Host);
 	char PortString[8];
