Include the trace store's local filesystem location in the /v1/status response.

diff --git a/Engine/Source/Developer/TraceAnalysis/Private/Store/AsioStore.cpp b/Engine/Source/Developer/TraceAnalysis/Private/Store/AsioStore.cpp
--- a/Engine/Source/Developer/TraceAnalysis/Private/Store/AsioStore.cpp
+++ b/Engine/Source/Developer/TraceAnalysis/Private/Store/AsioStore.cpp
@@ -221,6 +221,12 @@ void FAsioStore::WatchDir()
 	});
 }
 
+////////////////////////////////////////////////////////////////////////////////
+const TCHAR* FAsioStore::GetStoreDir() const
+{
+	return *StoreDir;
+}
+
 ////////////////////////////////////////////////////////////////////////////////
 uint32 FAsioStore::GetChangeSerial() const
 {
diff --git a/Engine/Source/Developer/TraceAnalysis/Private/Store/AsioStore.h b/Engine/Source/Developer/TraceAnalysis/Private/Store/AsioStore.h
--- a/Engine/Source/Developer/TraceAnalysis/Private/Store/AsioStore.h
+++ b/Engine/Source/Developer/TraceAnalysis/Private/Store/AsioStore.h
@@ -51,6 +51,7 @@ public:
 						FAsioStore(asio::io_context& IoContext, const TCHAR* InStoreDir);
 						~FAsioStore();
 	void				Close();
+	const TCHAR*		GetStoreDir() const;
 	uint32				GetChangeSerial() const;
 	uint32				GetTraceCount() const;
 	const FTrace*		GetTraceInfo(uint32 Index) const;
diff --git a/Engine/Source/Developer/TraceAnalysis/Private/Store/AsioStoreCborServer.cpp b/Engine/Source/Developer/TraceAnalysis/Private/Store/AsioStoreCborServer.cpp
--- a/Engine/Source/Developer/TraceAnalysis/Private/Store/AsioStoreCborServer.cpp
+++ b/Engine/Source/Developer/TraceAnalysis/Private/Store/AsioStoreCborServer.cpp
@@ -133,6 +133,7 @@ void FAsioStoreCborPeer::OnSessionInfo()
 void FAsioStoreCborPeer::OnStatus()
 {
 	TPayloadBuilder<> Builder(EStatusCode::Success);
+	Builder.AddString("store_dir", TCHAR_TO_UTF8(Store.GetStoreDir()));
 	Builder.AddInteger("recorder_port", Recorder.GetPort());
 	Builder.AddInteger("change_serial", Store.GetChangeSerial());
 	SendResponse(Builder.Done());
