Add accessor for the "store_dir" field of the trace store's /v1/status response.

diff --git a/Engine/Source/Developer/TraceAnalysis/Private/Store/StoreClient.cpp b/Engine/Source/Developer/TraceAnalysis/Private/Store/StoreClient.cpp
--- a/Engine/Source/Developer/TraceAnalysis/Private/Store/StoreClient.cpp
+++ b/Engine/Source/Developer/TraceAnalysis/Private/Store/StoreClient.cpp
@@ -375,6 +375,13 @@ bool FStoreCborClient::GetSessionInfoByTraceId(uint32 TraceId)
 namespace UE {
 namespace Trace {
 
+////////////////////////////////////////////////////////////////////////////////
+FAnsiStringView FStoreClient::FStatus::GetStoreDir() const
+{
+	const auto* Response = (const FResponse*)this;
+	return Response->GetString("store_dir", "");
+}
+
 ////////////////////////////////////////////////////////////////////////////////
 uint32 FStoreClient::FStatus::GetRecorderPort() const
 {
diff --git a/Engine/Source/Developer/TraceAnalysis/Public/Trace/StoreClient.h b/Engine/Source/Developer/TraceAnalysis/Public/Trace/StoreClient.h
--- a/Engine/Source/Developer/TraceAnalysis/Public/Trace/StoreClient.h
+++ b/Engine/Source/Developer/TraceAnalysis/Public/Trace/StoreClient.h
@@ -18,6 +18,7 @@ class TRACEANALYSIS_API FStoreClient
 public:
 	struct TRACEANALYSIS_API FStatus
 	{
+		FAnsiStringView	GetStoreDir() const;
 		uint32			GetRecorderPort() const;
 		uint32			GetChangeSerial() const;
 	};
