Tweaked the trace store client/server protocol to be more HTTP-like.

#rnx

diff --git a/Engine/Source/Developer/TraceAnalysis/Private/Store/AsioStoreCborServer.cpp b/Engine/Source/Developer/TraceAnalysis/Private/Store/AsioStoreCborServer.cpp
--- a/Engine/Source/Developer/TraceAnalysis/Private/Store/AsioStoreCborServer.cpp
+++ b/Engine/Source/Developer/TraceAnalysis/Private/Store/AsioStoreCborServer.cpp
@@ -218,8 +218,9 @@ void FAsioStoreCborPeer::OnTraceRead()
 ////////////////////////////////////////////////////////////////////////////////
 void FAsioStoreCborPeer::OnPayload()
 {
-	FAnsiStringView Method = Response.GetString("$method", "");
-	if (!Method.Len())
+	FAnsiStringView Request = Response.GetString("$request", "");
+	FAnsiStringView Path = Response.GetString("$path", "");
+	if (!Request.Len() || !Path.Len())
 	{
 		SendError(EStatusCode::BadRequest);
 		return;
@@ -237,16 +238,16 @@ void FAsioStoreCborPeer::OnPayload()
 		{ QuickStoreHash("v1/trace/read"),		&FAsioStoreCborPeer::OnTraceRead },
 	};
 
-	uint32 MethodHash = QuickStoreHash(Method);
+	uint32 PathHash = QuickStoreHash(Path);
 	for (const auto& Row : DispatchTable)
 	{
-		if (Row.Hash == MethodHash)
+		if (Row.Hash == PathHash)
 		{
 			return (this->*(Row.Func))();
 		}
 	}
 
-	SendError(EStatusCode::MethodNotAllowed);
+	SendError(EStatusCode::NotFound);
 }
 
 ////////////////////////////////////////////////////////////////////////////////
diff --git a/Engine/Source/Developer/TraceAnalysis/Private/Store/CborPayload.h b/Engine/Source/Developer/TraceAnalysis/Private/Store/CborPayload.h
--- a/Engine/Source/Developer/TraceAnalysis/Private/Store/CborPayload.h
+++ b/Engine/Source/Developer/TraceAnalysis/Private/Store/CborPayload.h
@@ -58,7 +58,7 @@ enum class EStatusCode
 	Unknown				= 0,
 	Success				= 200,
 	BadRequest			= 400,
-	MethodNotAllowed	= 405,
+	NotFound			= 404,
 	InternalError		= 500,
 };
 
@@ -70,7 +70,7 @@ class TPayloadBuilder
 {
 public:
 								TPayloadBuilder(EStatusCode StatusCode);
-	template <int N>			TPayloadBuilder(const char (&MethodName)[N]);
+	template <int N>			TPayloadBuilder(const char (&Path)[N]);
 	template <int N> void		AddInteger(const char (&Name)[N], int64 Value);
 	template <int N> void		AddString(const char (&Name)[N], const char* Value, int32 Length=-1);
 	FPayload					Done();
@@ -91,10 +91,11 @@ inline TPayloadBuilder<Size>::TPayloadBuilder(EStatusCode StatusCode)
 ////////////////////////////////////////////////////////////////////////////////
 template <int Size>
 template <int N>
-inline TPayloadBuilder<Size>::TPayloadBuilder(const char (&MethodName)[N])
+inline TPayloadBuilder<Size>::TPayloadBuilder(const char (&Path)[N])
 {
 	CborWriter.WriteContainerStart(ECborCode::Map, -1);
-	AddString("$method", MethodName, N - 1);
+	AddString("$request", "GET");
+	AddString("$path", Path, N - 1);
 }
 
 ////////////////////////////////////////////////////////////////////////////////
