Added some Vim fold markers to help navigate this file more easily.

#rnx

diff --git a/Engine/Source/Developer/TraceAnalysis/Private/Analysis/Engine.cpp b/Engine/Source/Developer/TraceAnalysis/Private/Analysis/Engine.cpp
--- a/Engine/Source/Developer/TraceAnalysis/Private/Analysis/Engine.cpp
+++ b/Engine/Source/Developer/TraceAnalysis/Private/Analysis/Engine.cpp
@@ -19,6 +19,8 @@
 namespace UE {
 namespace Trace {
 
+// {{{1 misc -------------------------------------------------------------------
+
 ////////////////////////////////////////////////////////////////////////////////
 void SerializeToCborImpl(TArray<uint8>&, const IAnalyzer::FEventData&, uint32);
 
@@ -41,6 +43,8 @@ private:
 
 
 
+// {{{1 aux-data ---------------------------------------------------------------
+
 ////////////////////////////////////////////////////////////////////////////////
 struct FAuxData
 {
@@ -58,6 +62,8 @@ struct FAnalysisEngine::FAuxDataCollector
 
 
 
+// {{{1 threads ----------------------------------------------------------------
+
 ////////////////////////////////////////////////////////////////////////////////
 FThreads::FThreads()
 {
@@ -122,6 +128,8 @@ const TArray<uint8>* FThreads::GetGroupName() const
 
 
 
+// {{{1 thread-info ------------------------------------------------------------
+
 ////////////////////////////////////////////////////////////////////////////////
 uint32 IAnalyzer::FThreadInfo::GetId() const
 {
@@ -169,6 +177,8 @@ const ANSICHAR* IAnalyzer::FThreadInfo::GetGroupName() const
 
 
 
+// {{{1 event-time -------------------------------------------------------------
+
 ////////////////////////////////////////////////////////////////////////////////
 uint64 IAnalyzer::FEventTime::GetTimestamp() const
 {
@@ -206,6 +216,8 @@ double IAnalyzer::FEventTime::AsSecondsAbsolute(int64 DurationCycles64) const
 
 
 
+// {{{1 dispatch ---------------------------------------------------------------
+
 ////////////////////////////////////////////////////////////////////////////////
 struct FAnalysisEngine::FDispatch
 {
@@ -258,6 +270,8 @@ int32 FAnalysisEngine::FDispatch::GetFieldIndex(const ANSICHAR* Name) const
 
 
 
+// {{{1 dispatch-builder -------------------------------------------------------
+
 ////////////////////////////////////////////////////////////////////////////////
 class FAnalysisEngine::FDispatchBuilder
 {
@@ -395,6 +409,8 @@ uint32 FAnalysisEngine::FDispatchBuilder::AppendName(const ANSICHAR* Name, int32
 
 
 
+// {{{1 event-type-info --------------------------------------------------------
+
 ////////////////////////////////////////////////////////////////////////////////
 uint32 IAnalyzer::FEventTypeInfo::GetId() const
 {
@@ -461,6 +477,8 @@ IAnalyzer::FEventFieldHandle IAnalyzer::FEventTypeInfo::GetFieldHandleImpl(
 
 
 
+// {{{1 field-info -------------------------------------------------------------
+
 ////////////////////////////////////////////////////////////////////////////////
 const ANSICHAR* IAnalyzer::FEventFieldInfo::GetName() const
 {
@@ -500,6 +518,8 @@ bool IAnalyzer::FEventFieldInfo::IsArray() const
 
 
 
+// {{{1 array-reader -----------------------------------------------------------
+
 ////////////////////////////////////////////////////////////////////////////////
 uint32 IAnalyzer::FArrayReader::Num() const
 {
@@ -526,6 +546,8 @@ const void* IAnalyzer::FArrayReader::GetImpl(uint32 Index, int16& SizeAndType) c
 
 
 
+// {{{1 event-data -------------------------------------------------------------
+
 ////////////////////////////////////////////////////////////////////////////////
 struct FAnalysisEngine::FEventDataInfo
 {
@@ -736,6 +758,8 @@ uint32 IAnalyzer::FEventData::GetAttachmentSize() const
 
 
 
+// {{{1 engine -----------------------------------------------------------------
+
 ////////////////////////////////////////////////////////////////////////////////
 enum ERouteId : uint16
 {
@@ -2017,5 +2041,8 @@ int32 FAnalysisEngine::OnDataProtocol4Impl(
 	return 0;
 }
 
+  // }}}
 } // namespace Trace
 } // namespace UE
+
+/* vim: set foldlevel=1 : */
