Delete redundant variable.

#rnx

diff --git a/Engine/Source/Developer/TraceAnalysis/Private/Analysis/Engine.cpp b/Engine/Source/Developer/TraceAnalysis/Private/Analysis/Engine.cpp
--- a/Engine/Source/Developer/TraceAnalysis/Private/Analysis/Engine.cpp
+++ b/Engine/Source/Developer/TraceAnalysis/Private/Analysis/Engine.cpp
@@ -779,7 +779,6 @@ enum ERouteId : uint16
 FAnalysisEngine::FAnalysisEngine(TArray<IAnalyzer*>&& InAnalyzers)
 : Analyzers(MoveTemp(InAnalyzers))
 {
-	uint16 SelfIndex = Analyzers.Num();
 	Analyzers.Add(this);
 }
 
