Removed a few redundant "using namespace" statements.

#rnx

diff --git a/Engine/Source/Developer/TraceAnalysis/Private/Analysis/Engine.cpp b/Engine/Source/Developer/TraceAnalysis/Private/Analysis/Engine.cpp
--- a/Engine/Source/Developer/TraceAnalysis/Private/Analysis/Engine.cpp
+++ b/Engine/Source/Developer/TraceAnalysis/Private/Analysis/Engine.cpp
@@ -790,8 +790,6 @@ FAnalysisEngine::~FAnalysisEngine()
 ////////////////////////////////////////////////////////////////////////////////
 void FAnalysisEngine::Begin()
 {
-	using namespace UE::Trace;
-
 	// Call out to all registered analyzers to have them register event interest
 	struct : IAnalyzer::FInterfaceBuilder
 	{
@@ -1168,8 +1166,6 @@ void FAnalysisEngine::ForEachRoute(uint32 RouteIndex, bool bScoped, ImplType&& I
 ////////////////////////////////////////////////////////////////////////////////
 void FAnalysisEngine::OnNewEventInternal(const void* EventData)
 {
-	using namespace UE::Trace;
-
 	FDispatchBuilder Builder;
 	switch (ProtocolVersion)
 	{
@@ -1248,8 +1244,6 @@ bool FAnalysisEngine::AddDispatch(FDispatch* Dispatch)
 ////////////////////////////////////////////////////////////////////////////////
 void FAnalysisEngine::OnNewEventProtocol0(FDispatchBuilder& Builder, const void* EventData)
 {
-	using namespace UE::Trace;
-
 	const auto& NewEvent = *(Protocol0::FNewEventEvent*)(EventData);
 
 	const auto* NameCursor = (const ANSICHAR*)(NewEvent.Fields + NewEvent.FieldCount);
@@ -1286,8 +1280,6 @@ void FAnalysisEngine::OnNewEventProtocol0(FDispatchBuilder& Builder, const void*
 ////////////////////////////////////////////////////////////////////////////////
 void FAnalysisEngine::OnNewEventProtocol1(FDispatchBuilder& Builder, const void* EventData)
 {
-	using namespace UE::Trace;
-
 	OnNewEventProtocol0(Builder, EventData);
 
 	const auto& NewEvent = *(Protocol1::FNewEventEvent*)(EventData);
@@ -1311,8 +1303,6 @@ void FAnalysisEngine::OnNewEventProtocol1(FDispatchBuilder& Builder, const void*
 ////////////////////////////////////////////////////////////////////////////////
 bool FAnalysisEngine::EstablishTransport(FStreamReader& Reader)
 {
-	using namespace UE::Trace;
-
 	const struct {
 		uint8 TransportVersion;
 		uint8 ProtocolVersion;
@@ -1429,8 +1419,6 @@ bool FAnalysisEngine::OnData(FStreamReader& Reader)
 ////////////////////////////////////////////////////////////////////////////////
 bool FAnalysisEngine::OnDataProtocol0()
 {
-	using namespace UE::Trace;
-
 	FThreads::FInfo ThreadInfo;
 
 	while (true)
@@ -1620,8 +1608,6 @@ bool FAnalysisEngine::OnDataProtocol2()
 ////////////////////////////////////////////////////////////////////////////////
 int32 FAnalysisEngine::OnDataProtocol2(FStreamReader& Reader, FThreads::FInfo& ThreadInfo)
 {
-	using namespace UE::Trace;
-
 	while (true)
 	{
 		auto Mark = Reader.SaveMark();
@@ -1734,8 +1720,6 @@ int32 FAnalysisEngine::OnDataProtocol2(FStreamReader& Reader, FThreads::FInfo& T
 ////////////////////////////////////////////////////////////////////////////////
 int32 FAnalysisEngine::OnDataProtocol2Aux(FStreamReader& Reader, FAuxDataCollector& Collector)
 {
-	using namespace UE::Trace;
-
 	while (true)
 	{
 		const uint8* NextByte = Reader.GetPointer<uint8>();
@@ -1795,7 +1779,7 @@ int32 FAnalysisEngine::OnDataProtocol4Known(
 	FStreamReader& Reader,
 	FThreads::FInfo& ThreadInfo)
 {
-	using namespace UE::Trace::Protocol4;
+	using namespace Protocol4;
 
 	switch (Uid)
 	{
@@ -1882,7 +1866,7 @@ int32 FAnalysisEngine::OnDataProtocol4Impl(
 {
 	/* Returns 0 if an event was successfully processed, 1 if there's not enough
 	 * data available, or ~AvailableLogSerial if the pending event is in the future */
-	using namespace UE::Trace::Protocol4;
+	using namespace Protocol4;
 
 	auto Mark = Reader.SaveMark();
 
