Moved a few classes out of FAnalysisEngine's scope.

#rnx

diff --git a/Engine/Source/Developer/TraceAnalysis/Private/Analysis/Engine.cpp b/Engine/Source/Developer/TraceAnalysis/Private/Analysis/Engine.cpp
--- a/Engine/Source/Developer/TraceAnalysis/Private/Analysis/Engine.cpp
+++ b/Engine/Source/Developer/TraceAnalysis/Private/Analysis/Engine.cpp
@@ -55,7 +55,7 @@ struct FAuxData
 };
 
 ////////////////////////////////////////////////////////////////////////////////
-struct FAnalysisEngine::FAuxDataCollector
+struct FAuxDataCollector
 : public TArray<FAuxData>
 {
 };
@@ -219,7 +219,7 @@ double IAnalyzer::FEventTime::AsSecondsAbsolute(int64 DurationCycles64) const
 // {{{1 dispatch ---------------------------------------------------------------
 
 ////////////////////////////////////////////////////////////////////////////////
-struct FAnalysisEngine::FDispatch
+struct FDispatch
 {
 	enum
 	{
@@ -252,7 +252,7 @@ struct FAnalysisEngine::FDispatch
 };
 
 ////////////////////////////////////////////////////////////////////////////////
-int32 FAnalysisEngine::FDispatch::GetFieldIndex(const ANSICHAR* Name) const
+int32 FDispatch::GetFieldIndex(const ANSICHAR* Name) const
 {
 	FFnv1aHash NameHash;
 	NameHash.Add(Name);
@@ -273,7 +273,7 @@ int32 FAnalysisEngine::FDispatch::GetFieldIndex(const ANSICHAR* Name) const
 // {{{1 dispatch-builder -------------------------------------------------------
 
 ////////////////////////////////////////////////////////////////////////////////
-class FAnalysisEngine::FDispatchBuilder
+class FDispatchBuilder
 {
 public:
 						FDispatchBuilder();
@@ -293,7 +293,7 @@ private:
 };
 
 ////////////////////////////////////////////////////////////////////////////////
-FAnalysisEngine::FDispatchBuilder::FDispatchBuilder()
+FDispatchBuilder::FDispatchBuilder()
 {
 	Buffer.SetNum(sizeof(FDispatch));
 
@@ -302,7 +302,7 @@ FAnalysisEngine::FDispatchBuilder::FDispatchBuilder()
 }
 
 ////////////////////////////////////////////////////////////////////////////////
-FAnalysisEngine::FDispatch* FAnalysisEngine::FDispatchBuilder::Finalize()
+FDispatch* FDispatchBuilder::Finalize()
 {
 	int32 Size = Buffer.Num() + NameBuffer.Num();
 	auto* Dispatch = (FDispatch*)FMemory::Malloc(Size);
@@ -332,49 +332,49 @@ FAnalysisEngine::FDispatch* FAnalysisEngine::FDispatchBuilder::Finalize()
 }
 
 ////////////////////////////////////////////////////////////////////////////////
-void FAnalysisEngine::FDispatchBuilder::SetUid(uint16 Uid)
+void FDispatchBuilder::SetUid(uint16 Uid)
 {
 	auto* Dispatch = (FDispatch*)(Buffer.GetData());
 	Dispatch->Uid = Uid;
 }
 
 ////////////////////////////////////////////////////////////////////////////////
-void FAnalysisEngine::FDispatchBuilder::SetLoggerName(const ANSICHAR* Name, int32 NameSize)
+void FDispatchBuilder::SetLoggerName(const ANSICHAR* Name, int32 NameSize)
 {
 	auto* Dispatch = (FDispatch*)(Buffer.GetData());
 	Dispatch->LoggerNameOffset += AppendName(Name, NameSize);
 }
 
 ////////////////////////////////////////////////////////////////////////////////
-void FAnalysisEngine::FDispatchBuilder::SetEventName(const ANSICHAR* Name, int32 NameSize)
+void FDispatchBuilder::SetEventName(const ANSICHAR* Name, int32 NameSize)
 {
 	auto* Dispatch = (FDispatch*)(Buffer.GetData());
 	Dispatch->EventNameOffset = AppendName(Name, NameSize);
 }
 
 ////////////////////////////////////////////////////////////////////////////////
-void FAnalysisEngine::FDispatchBuilder::SetImportant()
+void FDispatchBuilder::SetImportant()
 {
 	auto* Dispatch = (FDispatch*)(Buffer.GetData());
 	Dispatch->Flags |= FDispatch::Flag_Important;
 }
 
 ////////////////////////////////////////////////////////////////////////////////
-void FAnalysisEngine::FDispatchBuilder::SetMaybeHasAux()
+void FDispatchBuilder::SetMaybeHasAux()
 {
 	auto* Dispatch = (FDispatch*)(Buffer.GetData());
 	Dispatch->Flags |= FDispatch::Flag_MaybeHasAux;
 }
 
 ////////////////////////////////////////////////////////////////////////////////
-void FAnalysisEngine::FDispatchBuilder::SetNoSync()
+void FDispatchBuilder::SetNoSync()
 {
 	auto* Dispatch = (FDispatch*)(Buffer.GetData());
 	Dispatch->Flags |= FDispatch::Flag_NoSync;
 }
 
 ////////////////////////////////////////////////////////////////////////////////
-FAnalysisEngine::FDispatch::FField& FAnalysisEngine::FDispatchBuilder::AddField(const ANSICHAR* Name, int32 NameSize, uint16 Size)
+FDispatch::FField& FDispatchBuilder::AddField(const ANSICHAR* Name, int32 NameSize, uint16 Size)
 {
 	int32 Bufoff = Buffer.AddUninitialized(sizeof(FDispatch::FField));
 	auto* Field = (FDispatch::FField*)(Buffer.GetData() + Bufoff);
@@ -393,7 +393,7 @@ FAnalysisEngine::FDispatch::FField& FAnalysisEngine::FDispatchBuilder::AddField(
 }
 
 ////////////////////////////////////////////////////////////////////////////////
-uint32 FAnalysisEngine::FDispatchBuilder::AppendName(const ANSICHAR* Name, int32 NameSize)
+uint32 FDispatchBuilder::AppendName(const ANSICHAR* Name, int32 NameSize)
 {
 	if (NameSize < 0)
 	{
@@ -414,35 +414,35 @@ uint32 FAnalysisEngine::FDispatchBuilder::AppendName(const ANSICHAR* Name, int32
 ////////////////////////////////////////////////////////////////////////////////
 uint32 IAnalyzer::FEventTypeInfo::GetId() const
 {
-	const auto* Inner = (const FAnalysisEngine::FDispatch*)this;
+	const auto* Inner = (const FDispatch*)this;
 	return Inner->Uid;
 }
 
 ////////////////////////////////////////////////////////////////////////////////
 const ANSICHAR* IAnalyzer::FEventTypeInfo::GetName() const
 {
-	const auto* Inner = (const FAnalysisEngine::FDispatch*)this;
+	const auto* Inner = (const FDispatch*)this;
 	return (const ANSICHAR*)Inner + Inner->EventNameOffset;
 }
 
 ////////////////////////////////////////////////////////////////////////////////
 const ANSICHAR* IAnalyzer::FEventTypeInfo::GetLoggerName() const
 {
-	const auto* Inner = (const FAnalysisEngine::FDispatch*)this;
+	const auto* Inner = (const FDispatch*)this;
 	return (const ANSICHAR*)Inner + Inner->LoggerNameOffset;
 }
 
 ////////////////////////////////////////////////////////////////////////////////
 uint32 IAnalyzer::FEventTypeInfo::GetSize() const
 {
-	const auto* Inner = (const FAnalysisEngine::FDispatch*)this;
+	const auto* Inner = (const FDispatch*)this;
 	return Inner->EventSize;
 }
 
 ////////////////////////////////////////////////////////////////////////////////
 uint32 IAnalyzer::FEventTypeInfo::GetFieldCount() const
 {
-	const auto* Inner = (const FAnalysisEngine::FDispatch*)this;
+	const auto* Inner = (const FDispatch*)this;
 	return Inner->FieldCount;
 }
 
@@ -454,7 +454,7 @@ const IAnalyzer::FEventFieldInfo* IAnalyzer::FEventTypeInfo::GetFieldInfo(uint32
 		return nullptr;
 	}
 
-	const auto* Inner = (const FAnalysisEngine::FDispatch*)this;
+	const auto* Inner = (const FDispatch*)this;
 	return (const IAnalyzer::FEventFieldInfo*)(Inner->Fields + Index);
 }
 
@@ -463,14 +463,14 @@ IAnalyzer::FEventFieldHandle IAnalyzer::FEventTypeInfo::GetFieldHandleImpl(
 	const ANSICHAR* FieldName,
 	int16& SizeAndType) const
 {
-	const auto* Inner = (const FAnalysisEngine::FDispatch*)this;
+	const auto* Inner = (const FDispatch*)this;
 	int32 Index = Inner->GetFieldIndex(FieldName);
 	if (Index < 0)
 	{
 		return { -1 };
 	}
 
-	const FAnalysisEngine::FDispatch::FField& Field = Inner->Fields[Index];
+	const FDispatch::FField& Field = Inner->Fields[Index];
 	SizeAndType = Field.SizeAndType;
 	return { Field.Offset };
 }
@@ -482,14 +482,14 @@ IAnalyzer::FEventFieldHandle IAnalyzer::FEventTypeInfo::GetFieldHandleImpl(
 ////////////////////////////////////////////////////////////////////////////////
 const ANSICHAR* IAnalyzer::FEventFieldInfo::GetName() const
 {
-	const auto* Inner = (const FAnalysisEngine::FDispatch::FField*)this;
+	const auto* Inner = (const FDispatch::FField*)this;
 	return (const ANSICHAR*)(UPTRINT(Inner) + Inner->NameOffset);
 }
 
 ////////////////////////////////////////////////////////////////////////////////
 IAnalyzer::FEventFieldInfo::EType IAnalyzer::FEventFieldInfo::GetType() const
 {
-	const auto* Inner = (const FAnalysisEngine::FDispatch::FField*)this;
+	const auto* Inner = (const FDispatch::FField*)this;
 
 	if (Inner->Class == UE::Trace::Protocol0::Field_String)
 	{
@@ -512,7 +512,7 @@ IAnalyzer::FEventFieldInfo::EType IAnalyzer::FEventFieldInfo::GetType() const
 ////////////////////////////////////////////////////////////////////////////////
 bool IAnalyzer::FEventFieldInfo::IsArray() const
 {
-	const auto* Inner = (const FAnalysisEngine::FDispatch::FField*)this;
+	const auto* Inner = (const FDispatch::FField*)this;
 	return Inner->bArray;
 }
 
@@ -549,7 +549,7 @@ const void* IAnalyzer::FArrayReader::GetImpl(uint32 Index, int16& SizeAndType) c
 // {{{1 event-data -------------------------------------------------------------
 
 ////////////////////////////////////////////////////////////////////////////////
-struct FAnalysisEngine::FEventDataInfo
+struct FEventDataInfo
 {
 	const FAuxData*		GetAuxData(uint32 FieldIndex) const;
 	const uint8*		Ptr;
@@ -559,9 +559,9 @@ struct FAnalysisEngine::FEventDataInfo
 };
 
 ////////////////////////////////////////////////////////////////////////////////
-const FAuxData* FAnalysisEngine::FEventDataInfo::GetAuxData(uint32 FieldIndex) const
+const FAuxData* FEventDataInfo::GetAuxData(uint32 FieldIndex) const
 {
-	const auto* Info = (const FAnalysisEngine::FEventDataInfo*)this;
+	const auto* Info = (const FEventDataInfo*)this;
 	if (Info->AuxCollector == nullptr)
 	{
 		return nullptr;
@@ -582,14 +582,14 @@ const FAuxData* FAnalysisEngine::FEventDataInfo::GetAuxData(uint32 FieldIndex) c
 ////////////////////////////////////////////////////////////////////////////////
 const IAnalyzer::FEventTypeInfo& IAnalyzer::FEventData::GetTypeInfo() const
 {
-	const auto* Info = (const FAnalysisEngine::FEventDataInfo*)this;
+	const auto* Info = (const FEventDataInfo*)this;
 	return (const FEventTypeInfo&)(Info->Dispatch);
 }
 
 ////////////////////////////////////////////////////////////////////////////////
 const IAnalyzer::FArrayReader* IAnalyzer::FEventData::GetArrayImpl(const ANSICHAR* FieldName) const
 {
-	const auto* Info = (const FAnalysisEngine::FEventDataInfo*)this;
+	const auto* Info = (const FEventDataInfo*)this;
 
 	int32 Index = Info->Dispatch.GetFieldIndex(FieldName);
 	if (Index >= 0)
@@ -607,7 +607,7 @@ const IAnalyzer::FArrayReader* IAnalyzer::FEventData::GetArrayImpl(const ANSICHA
 ////////////////////////////////////////////////////////////////////////////////
 const void* IAnalyzer::FEventData::GetValueImpl(const ANSICHAR* FieldName, int16& SizeAndType) const
 {
-	const auto* Info = (const FAnalysisEngine::FEventDataInfo*)this;
+	const auto* Info = (const FEventDataInfo*)this;
 	const auto& Dispatch = Info->Dispatch;
 
 	int32 Index = Dispatch.GetFieldIndex(FieldName);
@@ -624,7 +624,7 @@ const void* IAnalyzer::FEventData::GetValueImpl(const ANSICHAR* FieldName, int16
 ////////////////////////////////////////////////////////////////////////////////
 bool IAnalyzer::FEventData::GetString(const ANSICHAR* FieldName, FAnsiStringView& Out) const
 {
-	const auto* Info = (const FAnalysisEngine::FEventDataInfo*)this;
+	const auto* Info = (const FEventDataInfo*)this;
 	const auto& Dispatch = Info->Dispatch;
 
 	int32 Index = Dispatch.GetFieldIndex(FieldName);
@@ -652,7 +652,7 @@ bool IAnalyzer::FEventData::GetString(const ANSICHAR* FieldName, FAnsiStringView
 ////////////////////////////////////////////////////////////////////////////////
 bool IAnalyzer::FEventData::GetString(const ANSICHAR* FieldName, FStringView& Out) const
 {
-	const auto* Info = (const FAnalysisEngine::FEventDataInfo*)this;
+	const auto* Info = (const FEventDataInfo*)this;
 	const auto& Dispatch = Info->Dispatch;
 
 	int32 Index = Dispatch.GetFieldIndex(FieldName);
@@ -680,7 +680,7 @@ bool IAnalyzer::FEventData::GetString(const ANSICHAR* FieldName, FStringView& Ou
 ////////////////////////////////////////////////////////////////////////////////
 bool IAnalyzer::FEventData::GetString(const ANSICHAR* FieldName, FString& Out) const
 {
-	const auto* Info = (const FAnalysisEngine::FEventDataInfo*)this;
+	const auto* Info = (const FEventDataInfo*)this;
 	const auto& Dispatch = Info->Dispatch;
 
 	int32 Index = Dispatch.GetFieldIndex(FieldName);
@@ -718,7 +718,7 @@ bool IAnalyzer::FEventData::GetString(const ANSICHAR* FieldName, FString& Out) c
 ////////////////////////////////////////////////////////////////////////////////
 uint32 IAnalyzer::FEventData::GetSize() const
 {
-	const auto* Info = (const FAnalysisEngine::FEventDataInfo*)this;
+	const auto* Info = (const FEventDataInfo*)this;
 	uint32 Size = Info->Size;
 	for (const FAuxData& Data : *(Info->AuxCollector))
 	{
@@ -730,7 +730,7 @@ uint32 IAnalyzer::FEventData::GetSize() const
 ////////////////////////////////////////////////////////////////////////////////
 void IAnalyzer::FEventData::SerializeToCbor(TArray<uint8>& Out) const
 {
-	const auto* Info = (const FAnalysisEngine::FEventDataInfo*)this;
+	const auto* Info = (const FEventDataInfo*)this;
 	uint32 Size = Info->Size;
 	if (Info->AuxCollector != nullptr)
 	{
@@ -745,14 +745,14 @@ void IAnalyzer::FEventData::SerializeToCbor(TArray<uint8>& Out) const
 ////////////////////////////////////////////////////////////////////////////////
 const uint8* IAnalyzer::FEventData::GetAttachment() const
 {
-	const auto* Info = (const FAnalysisEngine::FEventDataInfo*)this;
+	const auto* Info = (const FEventDataInfo*)this;
 	return Info->Ptr + Info->Dispatch.EventSize;
 }
 
 ////////////////////////////////////////////////////////////////////////////////
 uint32 IAnalyzer::FEventData::GetAttachmentSize() const
 {
-	const auto* Info = (const FAnalysisEngine::FEventDataInfo*)this;
+	const auto* Info = (const FEventDataInfo*)this;
 	return Info->Size - Info->Dispatch.EventSize;
 }
 
diff --git a/Engine/Source/Developer/TraceAnalysis/Private/Analysis/Engine.h b/Engine/Source/Developer/TraceAnalysis/Private/Analysis/Engine.h
--- a/Engine/Source/Developer/TraceAnalysis/Private/Analysis/Engine.h
+++ b/Engine/Source/Developer/TraceAnalysis/Private/Analysis/Engine.h
@@ -9,6 +9,9 @@
 namespace UE {
 namespace Trace {
 
+struct FAuxDataCollector;
+struct FDispatch;
+class FDispatchBuilder;
 class FStreamReader;
 
 ////////////////////////////////////////////////////////////////////////////////
@@ -57,9 +60,6 @@ class FAnalysisEngine
 	: public IAnalyzer
 {
 public:
-	struct				FDispatch;
-	struct				FAuxDataCollector;
-	struct				FEventDataInfo;
 						FAnalysisEngine(TArray<IAnalyzer*>&& InAnalyzers);
 						~FAnalysisEngine();
 	bool				OnData(FStreamReader& Reader);
@@ -85,7 +85,6 @@ private:
 		};
 	};
 
-	class				FDispatchBuilder;
 	void				Begin();
 	virtual bool		OnEvent(uint16 RouteId, EStyle Style, const FOnEventContext& Context) override;
 	virtual void		OnAnalysisBegin(const FOnAnalysisContext& Context) override;
