A strategic balancing of an End() call with a Begin() call.

#rnx

diff --git a/Engine/Source/Developer/TraceAnalysis/Private/Analysis/Engine.cpp b/Engine/Source/Developer/TraceAnalysis/Private/Analysis/Engine.cpp
--- a/Engine/Source/Developer/TraceAnalysis/Private/Analysis/Engine.cpp
+++ b/Engine/Source/Developer/TraceAnalysis/Private/Analysis/Engine.cpp
@@ -788,7 +788,7 @@ FAnalysisEngine::~FAnalysisEngine()
 }
 
 ////////////////////////////////////////////////////////////////////////////////
-void FAnalysisEngine::Begin()
+void FAnalysisEngine::OnTransportEstablished()
 {
 	// Call out to all registered analyzers to have them register event interest
 	struct : IAnalyzer::FInterfaceBuilder
@@ -896,6 +896,11 @@ void FAnalysisEngine::Begin()
 	}
 }
 
+////////////////////////////////////////////////////////////////////////////////
+void FAnalysisEngine::Begin()
+{
+}
+
 ////////////////////////////////////////////////////////////////////////////////
 void FAnalysisEngine::End()
 {
@@ -1397,7 +1402,7 @@ bool FAnalysisEngine::OnData(FStreamReader& Reader)
 			return false;
 		}
 
-		Begin();
+		OnTransportEstablished();
 	}
 
 	Transport->SetReader(Reader);
diff --git a/Engine/Source/Developer/TraceAnalysis/Private/Analysis/Engine.h b/Engine/Source/Developer/TraceAnalysis/Private/Analysis/Engine.h
--- a/Engine/Source/Developer/TraceAnalysis/Private/Analysis/Engine.h
+++ b/Engine/Source/Developer/TraceAnalysis/Private/Analysis/Engine.h
@@ -63,6 +63,7 @@ public:
 						FAnalysisEngine(TArray<IAnalyzer*>&& InAnalyzers);
 						~FAnalysisEngine();
 	bool				OnData(FStreamReader& Reader);
+	void				Begin();
 	void				End();
 
 private:
@@ -85,7 +86,7 @@ private:
 		};
 	};
 
-	void				Begin();
+	void				OnTransportEstablished();
 	virtual bool		OnEvent(uint16 RouteId, EStyle Style, const FOnEventContext& Context) override;
 	virtual void		OnAnalysisBegin(const FOnAnalysisContext& Context) override;
 	void				OnNewTrace(const FOnEventContext& Context);
diff --git a/Engine/Source/Developer/TraceAnalysis/Private/Analysis/Processor.cpp b/Engine/Source/Developer/TraceAnalysis/Private/Analysis/Processor.cpp
--- a/Engine/Source/Developer/TraceAnalysis/Private/Analysis/Processor.cpp
+++ b/Engine/Source/Developer/TraceAnalysis/Private/Analysis/Processor.cpp
@@ -34,8 +34,9 @@ FAnalysisProcessor::FImpl::~FImpl()
 ////////////////////////////////////////////////////////////////////////////////
 uint32 FAnalysisProcessor::FImpl::Run()
 {
-	FStreamBuffer Buffer;
+	AnalysisEngine.Begin();
 
+	FStreamBuffer Buffer;
 	while (!StopEvent->Wait(0, true))
 	{
 		UnpausedEvent->Wait();
