Over time the framework behind the analysis of traces has grown reactively and organically without little thought to structure. The monolithic class has been refactored to derive some structure that will ease future maintenance (and protocol changes in particular).

#rnx

diff --git a/Engine/Source/Developer/TraceAnalysis/Private/Analysis/Engine.cpp b/Engine/Source/Developer/TraceAnalysis/Private/Analysis/Engine.cpp
--- a/Engine/Source/Developer/TraceAnalysis/Private/Analysis/Engine.cpp
+++ b/Engine/Source/Developer/TraceAnalysis/Private/Analysis/Engine.cpp
@@ -64,6 +64,34 @@ struct FAuxDataCollector
 
 // {{{1 threads ----------------------------------------------------------------
 
+////////////////////////////////////////////////////////////////////////////////
+class FThreads
+{
+public:
+	struct FInfo
+	{
+		TArray<uint64>	ScopeRoutes;
+		TArray<uint8>	Name;
+		TArray<uint8>	GroupName;
+		int64			PrevTimestamp = 0;
+		uint32			ThreadId = ~0u;
+		uint32			SystemId = 0;
+		int32			SortHint = 0xff;
+	};
+
+						FThreads();
+						~FThreads();
+	FInfo*				GetInfo();
+	FInfo*				GetInfo(uint32 ThreadId);
+	void				SetGroupName(const ANSICHAR* InGroupName, uint32 Length);
+	const TArray<uint8>*GetGroupName() const;
+
+private:
+	uint32				LastGetInfoId = ~0u;
+	TArray<FInfo*>		Infos;
+	TArray<uint8>		GroupName;
+};
+
 ////////////////////////////////////////////////////////////////////////////////
 FThreads::FThreads()
 {
@@ -177,7 +205,18 @@ const ANSICHAR* IAnalyzer::FThreadInfo::GetGroupName() const
 
 
 
-// {{{1 event-time -------------------------------------------------------------
+// {{{1 timing -----------------------------------------------------------------
+
+////////////////////////////////////////////////////////////////////////////////
+struct FTiming
+{
+	uint64	BaseTimestamp;
+	uint64	TimestampHz;
+	double	InvTimestampHz;
+	uint64	EventTimestamp = 0;
+};
+
+
 
 ////////////////////////////////////////////////////////////////////////////////
 uint64 IAnalyzer::FEventTime::GetTimestamp() const
@@ -242,7 +281,6 @@ struct FDispatch
 	int32			GetFieldIndex(const ANSICHAR* Name) const;
 	uint32			Hash				= 0;
 	uint16			Uid					= 0;
-	uint16			FirstRoute			= ~uint16(0);
 	uint8			FieldCount			= 0;
 	uint8			Flags				= 0;
 	uint16			EventSize			= 0;
@@ -755,43 +793,303 @@ uint32 IAnalyzer::FEventData::GetAttachmentSize() const
 	const auto* Info = (const FEventDataInfo*)this;
 	return Info->Size - Info->Dispatch.EventSize;
 }
+// }}}
 
 
 
-// {{{1 engine -----------------------------------------------------------------
+// {{{1 type-registry ----------------------------------------------------------
+class FTypeRegistry
+{
+public:
+	typedef FDispatch FTypeInfo;
+
+						~FTypeRegistry();
+	const FTypeInfo*	Add(const void* TraceData);
+	const FTypeInfo*	Get(uint32 Uid) const;
+
+private:
+	TArray<FTypeInfo*>	TypeInfos;
+};
 
 ////////////////////////////////////////////////////////////////////////////////
-enum ERouteId : uint16
+FTypeRegistry::~FTypeRegistry()
 {
-	RouteId_NewEvent,
-	RouteId_NewTrace,
-	RouteId_SerialSync,
-	RouteId_Timing,
-	RouteId_ThreadTiming,
-	RouteId_ThreadInfo,
-	RouteId_ThreadGroupBegin,
-	RouteId_ThreadGroupEnd,
+	for (FTypeInfo* TypeInfo : TypeInfos)
+	{
+		FMemory::Free(TypeInfo);
+	}
+}
+
+////////////////////////////////////////////////////////////////////////////////
+const FTypeRegistry::FTypeInfo* FTypeRegistry::Get(uint32 Uid) const
+{
+	if (Uid >= uint32(TypeInfos.Num()))
+	{
+		return nullptr;
+	}
+
+	return TypeInfos[Uid];
+}
+
+////////////////////////////////////////////////////////////////////////////////
+const FTypeRegistry::FTypeInfo* FTypeRegistry::Add(const void* TraceData)
+{
+	FDispatchBuilder Builder;
+
+	const auto& NewEvent = *(Protocol4::FNewEventEvent*)(TraceData);
+
+	const auto* NameCursor = (const ANSICHAR*)(NewEvent.Fields + NewEvent.FieldCount);
+
+	Builder.SetLoggerName(NameCursor, NewEvent.LoggerNameSize);
+	NameCursor += NewEvent.LoggerNameSize;
+
+	Builder.SetEventName(NameCursor, NewEvent.EventNameSize);
+	NameCursor += NewEvent.EventNameSize;
+	Builder.SetUid(NewEvent.EventUid);
+
+	// Fill out the fields
+	for (int32 i = 0, n = NewEvent.FieldCount; i < n; ++i)
+	{
+		const auto& Field = NewEvent.Fields[i];
+
+		uint16 TypeSize = 1 << (Field.TypeInfo & Protocol0::Field_Pow2SizeMask);
+		if (Field.TypeInfo & Protocol0::Field_Float)
+		{
+			TypeSize = -TypeSize;
+		}
+
+		auto& OutField = Builder.AddField(NameCursor, Field.NameSize, Field.Size);
+		OutField.Offset = Field.Offset;
+		OutField.SizeAndType = TypeSize;
+
+		OutField.Class = Field.TypeInfo & Protocol0::Field_SpecialMask;
+		OutField.bArray = !!(Field.TypeInfo & Protocol0::Field_Array);
+
+		NameCursor += Field.NameSize;
+	}
+
+	if (NewEvent.Flags & uint32(Protocol4::EEventFlags::Important))
+	{
+		Builder.SetImportant();
+	}
+
+	if (NewEvent.Flags & uint32(Protocol4::EEventFlags::MaybeHasAux))
+	{
+		Builder.SetMaybeHasAux();
+	}
+
+	if (NewEvent.Flags & uint32(Protocol4::EEventFlags::NoSync))
+	{
+		Builder.SetNoSync();
+	}
+
+	FTypeInfo* TypeInfo = Builder.Finalize();
+
+	// Add the type to the type-infos table. Usually duplicates are an error
+	// but due to backwards compatibility we'll override existing types.
+	uint16 Uid = TypeInfo->Uid;
+	if (Uid < TypeInfos.Num())
+ 	{
+		if (TypeInfos[Uid] != nullptr)
+ 		{
+			FMemory::Free(TypeInfos[Uid]);
+			TypeInfos[Uid] = nullptr;
+ 		}
+ 	}
+	else
+ 	{
+		TypeInfos.SetNum(Uid + 1);
+ 	}
+
+	return TypeInfos[Uid] = TypeInfo;
+}
+
+// {{{1 analyzer-hub -----------------------------------------------------------
+class FAnalyzerHub
+{
+public:
+	void				End();
+	void				SetAnalyzers(TArray<IAnalyzer*>&& InAnalyzers);
+	void				OnNewType(const FTypeRegistry::FTypeInfo* TypeInfo);
+	void				OnEvent(const FTypeRegistry::FTypeInfo& TypeInfo, IAnalyzer::EStyle Style, const IAnalyzer::FOnEventContext& Context);
+	void				OnThreadInfo(const FThreads::FInfo& ThreadInfo);
+
+private:
+	void				BuildRoutes();
+	void				AddRoute(uint16 AnalyzerIndex, uint16 Id, const ANSICHAR* Logger, const ANSICHAR* Event, bool bScoped);
+	int32				GetRouteIndex(const FTypeRegistry::FTypeInfo& TypeInfo);
+	void				RetireAnalyzer(IAnalyzer* Analyzer);
+	template <typename ImplType>
+	void				ForEachRoute(uint32 RouteIndex, bool bScoped, ImplType&& Impl) const;
+
+	struct FRoute
+	{
+		uint32			Hash;
+		uint16			AnalyzerIndex : 15;
+		uint16			bScoped : 1;
+		uint16			Id;
+		union
+		{
+			uint32		ParentHash;
+			struct
+			{
+				int16	Count;
+				int16	Parent;
+			};
+		};
+	};
+
+	typedef TArray<uint16, TInlineAllocator<96>> TypeToRouteArray;
+
+	TArray<IAnalyzer*>	Analyzers;
+	TArray<FRoute>		Routes;
+	TypeToRouteArray	TypeToRoute; // biases by one so zero represents no route
 };
 
+////////////////////////////////////////////////////////////////////////////////
+void FAnalyzerHub::End()
+{
+	for (IAnalyzer* Analyzer : Analyzers)
+	{
+		if (Analyzer != nullptr)
+		{
+			Analyzer->OnAnalysisEnd();
+		}
+	}
+}
 
+////////////////////////////////////////////////////////////////////////////////
+void FAnalyzerHub::SetAnalyzers(TArray<IAnalyzer*>&& InAnalyzers)
+{
+	Analyzers = MoveTemp(InAnalyzers);
+	BuildRoutes();
+}
 
 ////////////////////////////////////////////////////////////////////////////////
-FAnalysisEngine::FAnalysisEngine(TArray<IAnalyzer*>&& InAnalyzers)
-: Analyzers(MoveTemp(InAnalyzers))
+int32 FAnalyzerHub::GetRouteIndex(const FTypeRegistry::FTypeInfo& TypeInfo)
 {
-	Analyzers.Add(this);
+	if (TypeInfo.Uid >= uint32(TypeToRoute.Num()))
+	{
+		return -1;
+	}
+
+	return int32(TypeToRoute[TypeInfo.Uid]) - 1;
 }
 
 ////////////////////////////////////////////////////////////////////////////////
-FAnalysisEngine::~FAnalysisEngine()
+template <typename ImplType>
+void FAnalyzerHub::ForEachRoute(uint32 RouteIndex, bool bScoped, ImplType&& Impl) const
+{
+	uint32 RouteCount = Routes.Num();
+	if (RouteIndex >= RouteCount)
+	{
+		return;
+	}
+
+	const FRoute* FirstRoute = Routes.GetData();
+	const FRoute* Route = FirstRoute + RouteIndex;
+	do
+	{
+		const FRoute* NextRoute = FirstRoute + Route->Parent;
+		for (uint32 n = Route->Count; n--; ++Route)
+		{
+			if (Route->bScoped != (bScoped == true))
+			{
+				continue;
+			}
+
+			IAnalyzer* Analyzer = Analyzers[Route->AnalyzerIndex];
+			if (Analyzer != nullptr)
+			{
+				Impl(Analyzer, Route->Id);
+			}
+		}
+		Route = NextRoute;
+	}
+	while (Route >= FirstRoute);
+}
+
+////////////////////////////////////////////////////////////////////////////////
+void FAnalyzerHub::OnNewType(const FTypeRegistry::FTypeInfo* TypeInfo)
+{
+	// Find routes that have subscribed to this event.
+	auto FindRoute = [this] (uint32 Hash)
+	{
+		int32 Index = Algo::LowerBoundBy(Routes, Hash, [] (const FRoute& Route) { return Route.Hash; });
+		return (Index < Routes.Num() && Routes[Index].Hash == Hash) ? Index : -1;
+	};
+
+	int32 FirstRoute = FindRoute(TypeInfo->Hash);
+	if (FirstRoute < 0)
+	{
+		FFnv1aHash LoggerHash;
+		LoggerHash.Add((const ANSICHAR*)TypeInfo + TypeInfo->LoggerNameOffset);
+		if ((FirstRoute = FindRoute(LoggerHash.Get())) < 0)
+		{
+			FirstRoute = FindRoute(FFnv1aHash().Get());
+		}
+	}
+
+	uint16 Uid = TypeInfo->Uid;
+	if (Uid >= uint16(TypeToRoute.Num()))
+ 	{
+		TypeToRoute.SetNumZeroed(Uid + 32);
+ 	}
+
+	TypeToRoute[Uid] = uint16(FirstRoute + 1);
+
+	// Inform routes that a new event has been declared.
+	if (FirstRoute >= 0)
+	{
+		ForEachRoute(FirstRoute, false, [&] (IAnalyzer* Analyzer, uint16 RouteId)
+		{
+			if (!Analyzer->OnNewEvent(RouteId, *(IAnalyzer::FEventTypeInfo*)TypeInfo))
+			{
+				RetireAnalyzer(Analyzer);
+			}
+		});
+	}
+}
+
+////////////////////////////////////////////////////////////////////////////////
+void FAnalyzerHub::OnEvent(
+	const FTypeRegistry::FTypeInfo& TypeInfo,
+	const IAnalyzer::EStyle Style,
+	const IAnalyzer::FOnEventContext& Context)
+{
+	int32 RouteIndex = GetRouteIndex(TypeInfo);
+	if (RouteIndex < 0)
+	{
+		return;
+	}
+
+	ForEachRoute(RouteIndex, false, [&] (IAnalyzer* Analyzer, uint16 RouteId)
+	{
+		if (!Analyzer->OnEvent(RouteId, Style, Context))
+		{
+			RetireAnalyzer(Analyzer);
+		}
+	});
+}
+
+////////////////////////////////////////////////////////////////////////////////
+void FAnalyzerHub::OnThreadInfo(const FThreads::FInfo& ThreadInfo)
 {
+	const auto& OuterThreadInfo = (IAnalyzer::FThreadInfo&)ThreadInfo;
+	for (IAnalyzer* Analyzer : Analyzers)
+	{
+		if (Analyzer != nullptr)
+		{
+			Analyzer->OnThreadInfo(OuterThreadInfo);
+		}
+	}
 }
 
 ////////////////////////////////////////////////////////////////////////////////
-void FAnalysisEngine::OnTransportEstablished()
+void FAnalyzerHub::BuildRoutes()
 {
 	// Call out to all registered analyzers to have them register event interest
-	struct : IAnalyzer::FInterfaceBuilder
+	struct : public IAnalyzer::FInterfaceBuilder
 	{
 		virtual void RouteEvent(uint16 RouteId, const ANSICHAR* Logger, const ANSICHAR* Event, bool bScoped) override
 		{
@@ -808,12 +1106,12 @@ void FAnalysisEngine::OnTransportEstablished()
 			Self->AddRoute(AnalyzerIndex, RouteId, "", "", bScoped);
 		}
 
-		FAnalysisEngine* Self;
+		FAnalyzerHub* Self;
 		uint16 AnalyzerIndex;
 	} Builder;
 	Builder.Self = this;
 
-	FOnAnalysisContext OnAnalysisContext = { Builder };
+	IAnalyzer::FOnAnalysisContext OnAnalysisContext = { Builder };
 	for (uint16 i = 0, n = Analyzers.Num(); i < n; ++i)
 	{
 		uint32 RouteCount = Routes.Num();
@@ -823,14 +1121,14 @@ void FAnalysisEngine::OnTransportEstablished()
 		Analyzer->OnAnalysisBegin(OnAnalysisContext);
 
 		// If the analyzer didn't add any routes we'll retire it immediately
-		if (RouteCount == Routes.Num() && Analyzer != this)
+		if (RouteCount == Routes.Num())
 		{
 			RetireAnalyzer(Analyzer);
 		}
 	}
 
+	// Sort routes by their hashes.
 	auto RouteProjection = [] (const FRoute& Route) { return Route.Hash; };
-
 	Algo::SortBy(Routes, RouteProjection);
 
 	auto FindRoute = [this, &RouteProjection] (uint32 Hash)
@@ -868,60 +1166,10 @@ void FAnalysisEngine::OnTransportEstablished()
 			Cursor = FixupRoute(Route);
 		}
 	}
-
-	switch (ProtocolVersion)
-	{
-	case Protocol0::EProtocol::Id:
-	case Protocol1::EProtocol::Id:
-	case Protocol2::EProtocol::Id:
-		{
-			FDispatchBuilder Dispatch;
-			Dispatch.SetUid(uint16(Protocol0::EKnownEventUids::NewEvent));
-			Dispatch.SetLoggerName("$Trace");
-			Dispatch.SetEventName("NewEvent");
-			AddDispatch(Dispatch.Finalize());
-		}
-		break;
-
-	case Protocol3::EProtocol::Id:
-		{
-			FDispatchBuilder Dispatch;
-			Dispatch.SetUid(uint16(Protocol3::EKnownEventUids::NewEvent));
-			Dispatch.SetLoggerName("$Trace");
-			Dispatch.SetEventName("NewEvent");
-			Dispatch.SetNoSync();
-			AddDispatch(Dispatch.Finalize());
-		}
-		break;
-	}
-}
-
-////////////////////////////////////////////////////////////////////////////////
-void FAnalysisEngine::Begin()
-{
-}
-
-////////////////////////////////////////////////////////////////////////////////
-void FAnalysisEngine::End()
-{
-	for (IAnalyzer* Analyzer : Analyzers)
-	{
-		if (Analyzer != nullptr)
-		{
-			Analyzer->OnAnalysisEnd();
-		}
-	}
-
-	for (FDispatch* Dispatch : Dispatches)
-	{
-		FMemory::Free(Dispatch);
-	}
-
-	delete Transport;
 }
 
 ////////////////////////////////////////////////////////////////////////////////
-void FAnalysisEngine::RetireAnalyzer(IAnalyzer* Analyzer)
+void FAnalyzerHub::RetireAnalyzer(IAnalyzer* Analyzer)
 {
 	for (uint32 i = 0, n = Analyzers.Num(); i < n; ++i)
 	{
@@ -937,7 +1185,7 @@ void FAnalysisEngine::RetireAnalyzer(IAnalyzer* Analyzer)
 }
 
 ////////////////////////////////////////////////////////////////////////////////
-void FAnalysisEngine::AddRoute(
+void FAnalyzerHub::AddRoute(
 	uint16 AnalyzerIndex,
 	uint16 Id,
 	const ANSICHAR* Logger,
@@ -966,30 +1214,93 @@ void FAnalysisEngine::AddRoute(
 	Route.bScoped = (bScoped == true);
 }
 
+// {{{1 state ------------------------------------------------------------------
+
 ////////////////////////////////////////////////////////////////////////////////
-void FAnalysisEngine::OnAnalysisBegin(const FOnAnalysisContext& Context)
+struct FAnalysisState
+{
+	struct FSerial
+	{
+		uint32	Value = 0;
+		uint32	Mask = 0;
+	};
+
+	FThreads	Threads;
+	FTiming		Timing;
+	FSerial		Serial;
+	uint32		UserUidBias = Protocol4::EKnownEventUids::User;
+};
+
+
+
+// {{{1 thread-info-cb ---------------------------------------------------------
+class FThreadInfoCallback
+{
+public:
+	virtual			~FThreadInfoCallback() {}
+	virtual void	OnThreadInfo(const FThreads::FInfo& Info) = 0;
+};
+
+
+
+// {{{1 analyzer ---------------------------------------------------------------
+
+////////////////////////////////////////////////////////////////////////////////
+enum
+{
+	RouteId_NewTrace,
+	RouteId_SerialSync,
+	RouteId_Timing,
+	RouteId_ThreadTiming,
+	RouteId_ThreadInfo,
+	RouteId_ThreadGroupBegin,
+	RouteId_ThreadGroupEnd,
+};
+
+////////////////////////////////////////////////////////////////////////////////
+class FTraceAnalyzer
+	: public IAnalyzer
+{
+public:
+							FTraceAnalyzer(FAnalysisState& InState, FThreadInfoCallback& InCallback);
+	virtual bool			OnEvent(uint16 RouteId, EStyle Style, const FOnEventContext& Context) override;
+	virtual void			OnAnalysisBegin(const FOnAnalysisContext& Context) override;
+	void					OnNewTrace(const FOnEventContext& Context);
+	void					OnSerialSync(const FOnEventContext& Context);
+	void					OnThreadTiming(const FOnEventContext& Context);
+	void					OnThreadInfoInternal(const FOnEventContext& Context);
+	void					OnThreadGroupBegin(const FOnEventContext& Context);
+	void					OnThreadGroupEnd();
+	void					OnTiming(const FOnEventContext& Context);
+
+private:
+	FAnalysisState&			State;
+	FThreadInfoCallback&	ThreadInfoCallback;
+};
+
+////////////////////////////////////////////////////////////////////////////////
+FTraceAnalyzer::FTraceAnalyzer(FAnalysisState& InState, FThreadInfoCallback& InCallback)
+: State(InState)
+, ThreadInfoCallback(InCallback)
 {
-	auto& Builder = Context.InterfaceBuilder;
-	Builder.RouteEvent(RouteId_NewTrace,		"$Trace", "NewTrace");
-	Builder.RouteEvent(RouteId_SerialSync,		"$Trace", "SerialSync");
-	Builder.RouteEvent(RouteId_NewEvent,		"$Trace", "NewEvent");
-	Builder.RouteEvent(RouteId_Timing,			"$Trace", "Timing");
-	Builder.RouteEvent(RouteId_ThreadTiming,	"$Trace", "ThreadTiming");
-	Builder.RouteEvent(RouteId_ThreadInfo,		"$Trace", "ThreadInfo");
-	Builder.RouteEvent(RouteId_ThreadGroupBegin,"$Trace", "ThreadGroupBegin");
-	Builder.RouteEvent(RouteId_ThreadGroupEnd,	"$Trace", "ThreadGroupEnd");
 }
 
 ////////////////////////////////////////////////////////////////////////////////
-bool FAnalysisEngine::OnEvent(uint16 RouteId, EStyle Style, const FOnEventContext& Context)
+void FTraceAnalyzer::OnAnalysisBegin(const FOnAnalysisContext& Context)
 {
-	if (RouteId == RouteId_NewEvent)
-	{
-		const FEventDataInfo& EventData = (const FEventDataInfo&)(Context.EventData);
-		OnNewEventInternal(EventData.Ptr);
-		return true;
-	}
+	auto& Builder = Context.InterfaceBuilder;
+	Builder.RouteEvent(RouteId_NewTrace,		"$Trace", "NewTrace");
+	Builder.RouteEvent(RouteId_SerialSync,		"$Trace", "SerialSync");
+	Builder.RouteEvent(RouteId_Timing,			"$Trace", "Timing");
+	Builder.RouteEvent(RouteId_ThreadTiming,	"$Trace", "ThreadTiming");
+	Builder.RouteEvent(RouteId_ThreadInfo,		"$Trace", "ThreadInfo");
+	Builder.RouteEvent(RouteId_ThreadGroupBegin,"$Trace", "ThreadGroupBegin");
+	Builder.RouteEvent(RouteId_ThreadGroupEnd,	"$Trace", "ThreadGroupEnd");
+}
 
+////////////////////////////////////////////////////////////////////////////////
+bool FTraceAnalyzer::OnEvent(uint16 RouteId, EStyle Style, const FOnEventContext& Context)
+{
 	switch (RouteId)
 	{
 	case RouteId_NewTrace:			OnNewTrace(Context);			break;
@@ -1005,13 +1316,14 @@ bool FAnalysisEngine::OnEvent(uint16 RouteId, EStyle Style, const FOnEventContex
 }
 
 ////////////////////////////////////////////////////////////////////////////////
-void FAnalysisEngine::OnNewTrace(const FOnEventContext& Context)
+void FTraceAnalyzer::OnNewTrace(const FOnEventContext& Context)
 {
 	const FEventData& EventData = Context.EventData;
 
 	// "Serial" will tell us approximately where we've started in the log serial
 	// range. We'll bias it by half so we won't accept any serialised events and
 	// mark the MSB to indicate that the current serial should be corrected.
+	auto& Serial = State.Serial;
 	uint32 Hint = EventData.GetValue<uint32>("Serial");
 	Hint -= (Serial.Mask + 1) >> 1;
 	Hint &= Serial.Mask;
@@ -1025,52 +1337,45 @@ void FAnalysisEngine::OnNewTrace(const FOnEventContext& Context)
 		OnSerialSync(Context);
 	}
 
-	UserUidBias = EventData.GetValue<uint32>("UserUidBias", uint32(UE::Trace::Protocol3::EKnownEventUids::User));
+	State.UserUidBias = EventData.GetValue<uint32>("UserUidBias", uint32(UE::Trace::Protocol3::EKnownEventUids::User));
 
 	OnTiming(Context);
 }
 
 ////////////////////////////////////////////////////////////////////////////////
-void FAnalysisEngine::OnSerialSync(const FOnEventContext& Context)
+void FTraceAnalyzer::OnSerialSync(const FOnEventContext& Context)
 {
-	Serial.Value &= ~0x40000000;
+	State.Serial.Value &= ~0x40000000;
 }
 
 ////////////////////////////////////////////////////////////////////////////////
-void FAnalysisEngine::OnTiming(const FOnEventContext& Context)
+void FTraceAnalyzer::OnTiming(const FOnEventContext& Context)
 {
 	uint64 StartCycle = Context.EventData.GetValue<uint64>("StartCycle");
 	uint64 CycleFrequency = Context.EventData.GetValue<uint64>("CycleFrequency");
 
-	Timing.BaseTimestamp = StartCycle;
-	Timing.TimestampHz = CycleFrequency;
-	Timing.InvTimestampHz = 1.0 / double(CycleFrequency);
+	State.Timing.BaseTimestamp = StartCycle;
+	State.Timing.TimestampHz = CycleFrequency;
+	State.Timing.InvTimestampHz = 1.0 / double(CycleFrequency);
 }
 
 ////////////////////////////////////////////////////////////////////////////////
-void FAnalysisEngine::OnThreadTiming(const FOnEventContext& Context)
+void FTraceAnalyzer::OnThreadTiming(const FOnEventContext& Context)
 {
 	uint64 BaseTimestamp = Context.EventData.GetValue<uint64>("BaseTimestamp");
-	if (FThreads::FInfo* Info = Threads.GetInfo())
+	if (FThreads::FInfo* Info = State.Threads.GetInfo())
 	{
 		Info->PrevTimestamp = BaseTimestamp;
 
 		// We can springboard of this event as a way to know a thread has just
 		// started (or at least is about to send its first event). Notify analyzers
 		// so they're aware of threads that never get explicitly registered.
-		const auto* OuterInfo = (FThreadInfo*)Info;
-		for (uint16 i = 0, n = Analyzers.Num(); i < n; ++i)
-		{
-			if (IAnalyzer* Analyzer = Analyzers[i])
-			{
-				Analyzer->OnThreadInfo(*OuterInfo);
-			}
-		}
+		ThreadInfoCallback.OnThreadInfo(*Info);
 	}
 }
 
 ////////////////////////////////////////////////////////////////////////////////
-void FAnalysisEngine::OnThreadInfoInternal(const FOnEventContext& Context)
+void FTraceAnalyzer::OnThreadInfoInternal(const FOnEventContext& Context)
 {
 	const FEventData& EventData = Context.EventData;
 
@@ -1080,11 +1385,11 @@ void FAnalysisEngine::OnThreadInfoInternal(const FOnEventContext& Context)
 	{
 		// Post important-events; the thread-info event is not on the thread it
 		// represents anymore. Fortunately the thread-id is traced now.
-		ThreadInfo = Threads.GetInfo(ThreadId);
+		ThreadInfo = State.Threads.GetInfo(ThreadId);
 	}
 	else
 	{
-		ThreadInfo = Threads.GetInfo();
+		ThreadInfo = State.Threads.GetInfo();
 	}
 
 	if (ThreadInfo == nullptr)
@@ -1103,328 +1408,331 @@ void FAnalysisEngine::OnThreadInfoInternal(const FOnEventContext& Context)
 
 	if (ThreadInfo->GroupName.Num() <= 0)
 	{
-		if (const TArray<uint8>* GroupName = Threads.GetGroupName())
+		if (const TArray<uint8>* GroupName = State.Threads.GetGroupName())
 		{
 			ThreadInfo->GroupName = *GroupName;
 		}
 	}
 
-	const auto* OuterInfo = (FThreadInfo*)ThreadInfo;
-	for (uint16 i = 0, n = Analyzers.Num(); i < n; ++i)
-	{
-		if (IAnalyzer* Analyzer = Analyzers[i])
-		{
-			Analyzer->OnThreadInfo(*OuterInfo);
-		}
-	}
+	ThreadInfoCallback.OnThreadInfo(*ThreadInfo);
 }
 
 ////////////////////////////////////////////////////////////////////////////////
-void FAnalysisEngine::OnThreadGroupBegin(const FOnEventContext& Context)
+void FTraceAnalyzer::OnThreadGroupBegin(const FOnEventContext& Context)
 {
 	const FEventData& EventData = Context.EventData;
 
 	FAnsiStringView Name;
 	EventData.GetString("Name", Name);
-	Threads.SetGroupName(Name.GetData(), Name.Len());
+	State.Threads.SetGroupName(Name.GetData(), Name.Len());
 }
 
 ////////////////////////////////////////////////////////////////////////////////
-void FAnalysisEngine::OnThreadGroupEnd()
+void FTraceAnalyzer::OnThreadGroupEnd()
 {
-	Threads.SetGroupName("", 0);
+	State.Threads.SetGroupName("", 0);
 }
 
+// {{{1 bridge -----------------------------------------------------------------
+class FAnalysisBridge
+	: public FThreadInfoCallback
+{
+public:
+	typedef FAnalysisState::FSerial FSerial;
+
+						FAnalysisBridge(TArray<IAnalyzer*>&& Analyzers);
+	void				Reset();
+	uint32				GetUserUidBias() const;
+	FSerial&			GetSerial();
+	void				SetActiveThread(uint32 ThreadId);
+	void				OnNewType(const FTypeRegistry::FTypeInfo* TypeInfo);
+	void				OnEvent(const FEventDataInfo& EventDataInfo);
+	virtual void		OnThreadInfo(const FThreads::FInfo& InThreadInfo) override;
+	void				EnterScope(uint64 Timestamp=0);
+	void				LeaveScope(uint64 Timestamp=0);
+
+private:
+	FAnalysisState		State;
+	FTraceAnalyzer		TraceAnalyzer = { State, *this };
+	FAnalyzerHub		AnalyzerHub;
+	FThreads::FInfo*	ThreadInfo = nullptr;
+};
+
 ////////////////////////////////////////////////////////////////////////////////
-template <typename ImplType>
-void FAnalysisEngine::ForEachRoute(uint32 RouteIndex, bool bScoped, ImplType&& Impl)
+FAnalysisBridge::FAnalysisBridge(TArray<IAnalyzer*>&& Analyzers)
 {
-	uint32 RouteCount = Routes.Num();
-	if (RouteIndex >= RouteCount)
-	{
-		return;
-	}
+	TArray<IAnalyzer*> TempAnalyzers(MoveTemp(Analyzers));
+	TempAnalyzers.Add(&TraceAnalyzer);
+	AnalyzerHub.SetAnalyzers(MoveTemp(TempAnalyzers));
+}
 
-	const FRoute* FirstRoute = Routes.GetData();
-	const FRoute* Route = FirstRoute + RouteIndex;
-	do
-	{
-		const FRoute* NextRoute = FirstRoute + Route->Parent;
-		for (uint32 n = Route->Count; n--; ++Route)
-		{
-			if (Route->bScoped != (bScoped == true))
-			{
-				continue;
-			}
+////////////////////////////////////////////////////////////////////////////////
+void FAnalysisBridge::Reset()
+{
+	AnalyzerHub.End();
 
-			IAnalyzer* Analyzer = Analyzers[Route->AnalyzerIndex];
-			if (Analyzer != nullptr)
-			{
-				Impl(Analyzer, Route->Id);
-			}
-		}
-		Route = NextRoute;
-	}
-	while (Route >= FirstRoute);
+	State.~FAnalysisState();
+	new (&State) FAnalysisState();
 }
 
 ////////////////////////////////////////////////////////////////////////////////
-void FAnalysisEngine::OnNewEventInternal(const void* EventData)
+void FAnalysisBridge::SetActiveThread(uint32 ThreadId)
 {
-	FDispatchBuilder Builder;
-	switch (ProtocolVersion)
-	{
-	case Protocol0::EProtocol::Id:
-		OnNewEventProtocol0(Builder, EventData);
-		break;
+	ThreadInfo = State.Threads.GetInfo(ThreadId);
+}
 
-	case Protocol1::EProtocol::Id:
-	case Protocol2::EProtocol::Id:
-	case Protocol3::EProtocol::Id:
-	case Protocol4::EProtocol::Id:
-		OnNewEventProtocol1(Builder, EventData);
-		break;
-	}
+////////////////////////////////////////////////////////////////////////////////
+uint32 FAnalysisBridge::GetUserUidBias() const
+{
+	return State.UserUidBias;
+}
 
-	// Get the dispatch and add it into the dispatch table. Fail gently if there
-	// is the dispatch table unexpetedly already has an entry.
-	FDispatch* Dispatch = Builder.Finalize();
-	if (!AddDispatch(Dispatch))
-	{
-		return;
-	}
+////////////////////////////////////////////////////////////////////////////////
+FAnalysisBridge::FSerial& FAnalysisBridge::GetSerial()
+{
+	return State.Serial;
+}
 
-	// Inform routes that a new event has been declared.
-	ForEachRoute(Dispatch->FirstRoute, false, [&] (IAnalyzer* Analyzer, uint16 RouteId)
-	{
-		if (!Analyzer->OnNewEvent(RouteId, *(FEventTypeInfo*)Dispatch))
-		{
-			RetireAnalyzer(Analyzer);
-		}
-	});
+////////////////////////////////////////////////////////////////////////////////
+void FAnalysisBridge::OnNewType(const FTypeRegistry::FTypeInfo* TypeInfo)
+{
+	AnalyzerHub.OnNewType(TypeInfo);
 }
 
 ////////////////////////////////////////////////////////////////////////////////
-bool FAnalysisEngine::AddDispatch(FDispatch* Dispatch)
+void FAnalysisBridge::OnEvent(const FEventDataInfo& EventDataInfo)
 {
-	// Add the dispatch to the dispatch table. Usually duplicates are an error
-	// but due to backwards compatibility we'll override existing dispatches.
-	uint16 Uid = Dispatch->Uid;
-	if (Uid < Dispatches.Num())
- 	{
-		if (Dispatches[Uid] != nullptr)
- 		{
-			FMemory::Free(Dispatches[Uid]);
-			Dispatches[Uid] = nullptr;
- 		}
- 	}
+	IAnalyzer::EStyle Style = IAnalyzer::EStyle::Normal;
+	if (ThreadInfo->ScopeRoutes.Num() > 0 && int64(ThreadInfo->ScopeRoutes.Last()) < 0)
+	{
+		Style = IAnalyzer::EStyle::EnterScope;
+		State.Timing.EventTimestamp = ~(ThreadInfo->ScopeRoutes.Last());
+	}
 	else
- 	{
-		Dispatches.SetNum(Uid + 1);
- 	}
-
-	// Find routes that have subscribed to this event.
-	auto FindRoute = [this] (uint32 Hash)
 	{
-		int32 Index = Algo::LowerBoundBy(Routes, Hash, [] (const FRoute& Route) { return Route.Hash; });
-		return (Index < Routes.Num() && Routes[Index].Hash == Hash) ? Index : -1;
+		State.Timing.EventTimestamp = 0;
+	}
+
+	// TODO "Dispatch" should be renamed "EventTypeInfo" or similar.
+	const FTypeRegistry::FTypeInfo* TypeInfo = &(EventDataInfo.Dispatch);
+
+	IAnalyzer::FOnEventContext Context = {
+		*(const IAnalyzer::FThreadInfo*)ThreadInfo,
+		(const IAnalyzer::FEventTime&)(State.Timing),
+		(const IAnalyzer::FEventData&)EventDataInfo,
 	};
+	AnalyzerHub.OnEvent(*TypeInfo, Style, Context);
 
-	int32 FirstRoute = FindRoute(Dispatch->Hash);
-	if (FirstRoute < 0)
+	if (Style == IAnalyzer::EStyle::EnterScope)
 	{
-		FFnv1aHash LoggerHash;
-		LoggerHash.Add((const ANSICHAR*)Dispatch + Dispatch->LoggerNameOffset);
-		if ((FirstRoute = FindRoute(LoggerHash.Get())) < 0)
-		{
-			FirstRoute = FindRoute(FFnv1aHash().Get());
-		}
+		ThreadInfo->ScopeRoutes.Last() = PTRINT(TypeInfo);
 	}
-	Dispatch->FirstRoute = FirstRoute;
+}
 
-	Dispatches[Uid] = Dispatch;
-	return true;
+////////////////////////////////////////////////////////////////////////////////
+void FAnalysisBridge::OnThreadInfo(const FThreads::FInfo& InThreadInfo)
+{
+	// Note that InThreadInfo might not equal the bridge's ThreadInfo because
+	// information about threads comes from trace and could have been traced on
+	// a different thread to the one it is describing (or no thread at all in
+	// the case of important events).
+	AnalyzerHub.OnThreadInfo(InThreadInfo);
 }
 
 ////////////////////////////////////////////////////////////////////////////////
-void FAnalysisEngine::OnNewEventProtocol0(FDispatchBuilder& Builder, const void* EventData)
+void FAnalysisBridge::EnterScope(uint64 Timestamp)
 {
-	const auto& NewEvent = *(Protocol0::FNewEventEvent*)(EventData);
+    if (Timestamp)
+    {
+	    Timestamp = ThreadInfo->PrevTimestamp += Timestamp;
+    }
 
-	const auto* NameCursor = (const ANSICHAR*)(NewEvent.Fields + NewEvent.FieldCount);
+    ThreadInfo->ScopeRoutes.Push(~Timestamp);
+}
 
-	Builder.SetLoggerName(NameCursor, NewEvent.LoggerNameSize);
-	NameCursor += NewEvent.LoggerNameSize;
+////////////////////////////////////////////////////////////////////////////////
+void FAnalysisBridge::LeaveScope(uint64 Timestamp)
+{
+    if (Timestamp)
+    {
+	    Timestamp = ThreadInfo->PrevTimestamp += Timestamp;
+    }
 
-	Builder.SetEventName(NameCursor, NewEvent.EventNameSize);
-	NameCursor += NewEvent.EventNameSize;
-	Builder.SetUid(NewEvent.EventUid);
+	State.Timing.EventTimestamp = Timestamp;
 
-	// Fill out the fields
-	for (int i = 0, n = NewEvent.FieldCount; i < n; ++i)
+	if (ThreadInfo->ScopeRoutes.Num() > 0)
 	{
-		const auto& Field = NewEvent.Fields[i];
-
-		uint16 TypeSize = 1 << (Field.TypeInfo & Protocol0::Field_Pow2SizeMask);
-		if (Field.TypeInfo & Protocol0::Field_Float)
-		{
-			TypeSize = -TypeSize;
-		}
+		PTRINT TypeInfoPtr = ThreadInfo->ScopeRoutes.Pop(false);
+		const auto* TypeInfo = (FTypeRegistry::FTypeInfo*)TypeInfoPtr;
 
-		auto& OutField = Builder.AddField(NameCursor, Field.NameSize, Field.Size);
-		OutField.Offset = Field.Offset;
-		OutField.SizeAndType = TypeSize;
+		FEventDataInfo EmptyEventInfo = {
+			nullptr,
+			*TypeInfo
+		};
 
-		OutField.Class = Field.TypeInfo & Protocol0::Field_SpecialMask;
-		OutField.bArray = !!(Field.TypeInfo & Protocol0::Field_Array);
+		IAnalyzer::FOnEventContext Context = {
+			(const IAnalyzer::FThreadInfo&)ThreadInfo,
+			(const IAnalyzer::FEventTime&)(State.Timing),
+			(const IAnalyzer::FEventData&)EmptyEventInfo,
+		};
 
-		NameCursor += Field.NameSize;
+		AnalyzerHub.OnEvent(*TypeInfo, IAnalyzer::EStyle::LeaveScope, Context);
 	}
+
+	State.Timing.EventTimestamp = 0;
 }
 
-////////////////////////////////////////////////////////////////////////////////
-void FAnalysisEngine::OnNewEventProtocol1(FDispatchBuilder& Builder, const void* EventData)
-{
-	OnNewEventProtocol0(Builder, EventData);
 
-	const auto& NewEvent = *(Protocol1::FNewEventEvent*)(EventData);
 
-	if (NewEvent.Flags & int(Protocol1::EEventFlags::Important))
+// {{{1 machine ----------------------------------------------------------------
+////////////////////////////////////////////////////////////////////////////////
+class FAnalysisMachine
+{
+public:
+	enum class EStatus
 	{
-		Builder.SetImportant();
-	}
+		Error,
+		Abort,
+		NotEnoughData,
+		Eof,
+		Continue,
+	};
 
-	if (NewEvent.Flags & int(Protocol1::EEventFlags::MaybeHasAux))
+	struct FMachineContext
 	{
-		Builder.SetMaybeHasAux();
-	}
+		FAnalysisMachine&	Machine;
+		FAnalysisBridge&	Bridge;
+	};
 
-	if (NewEvent.Flags & int(Protocol1::EEventFlags::NoSync))
+	class FStage
 	{
-		Builder.SetNoSync();
-	}
+	public:
+		typedef FAnalysisMachine::FMachineContext	FMachineContext;
+		typedef FAnalysisMachine::EStatus			EStatus;
+
+		virtual				~FStage() {}
+		virtual EStatus		OnData(FStreamReader& Reader, const FMachineContext& Context) = 0;
+		virtual void		EnterStage(const FMachineContext& Context) {};
+		virtual void		ExitStage(const FMachineContext& Context) {};
+	};
+
+							FAnalysisMachine(FAnalysisBridge& InBridge);
+							~FAnalysisMachine();
+	EStatus					OnData(FStreamReader& Reader);
+	void					Transition();
+	template <class StageType, typename... ArgsType>
+	StageType*				QueueStage(ArgsType... Args);
+
+private:
+	void					CleanUp();
+	FAnalysisBridge&		Bridge;
+	FStage*					ActiveStage = nullptr;
+	TArray<FStage*>			StageQueue;
+	TArray<FStage*>			DeadStages;
+};
+
+////////////////////////////////////////////////////////////////////////////////
+FAnalysisMachine::FAnalysisMachine(FAnalysisBridge& InBridge)
+: Bridge(InBridge)
+{
 }
 
 ////////////////////////////////////////////////////////////////////////////////
-bool FAnalysisEngine::EstablishTransport(FStreamReader& Reader)
+FAnalysisMachine::~FAnalysisMachine()
 {
-	const struct {
-		uint8 TransportVersion;
-		uint8 ProtocolVersion;
-	}* Header = decltype(Header)(Reader.GetPointer(sizeof(*Header)));
-	if (Header == nullptr)
-	{
-		return false;
-	}
+	CleanUp();
+}
 
-	// Check for the magic uint32. Early traces did not include this as it was
-	// used to validate a inbound socket connection and then discarded.
-	if (Header->TransportVersion == 'E' ||
-		Header->TransportVersion == 'T' ||
-		Header->TransportVersion == '2')
+////////////////////////////////////////////////////////////////////////////////
+void FAnalysisMachine::CleanUp()
+{
+	for (FStage* Stage : DeadStages)
 	{
-		const uint32* Magic = (const uint32*)(Reader.GetPointer(sizeof(*Magic)));
-		if (*Magic == 'ECRT' || *Magic == '2CRT')
-		{
-			// Source is big-endian which we don't currently support
-			return false;
-		}
+		delete Stage;
+	}
+	DeadStages.Reset();
+}
 
-		if (*Magic == 'TRCE')
-		{
-			Reader.Advance(sizeof(*Magic));
-			return EstablishTransport(Reader);
-		}
+////////////////////////////////////////////////////////////////////////////////
+template <class StageType, typename... ArgsType>
+StageType* FAnalysisMachine::QueueStage(ArgsType... Args)
+{
+	StageType* Stage = new StageType(Args...);
+	StageQueue.Insert(Stage, 0);
+	return Stage;
+}
 
-		if (*Magic == 'TRC2')
-		{
-			Reader.Advance(sizeof(*Magic));
-			if (const uint16* MetadataSize = Reader.GetPointer<uint16>())
-			{
-				Reader.Advance(sizeof(*MetadataSize) + *MetadataSize);
-				return EstablishTransport(Reader);
-			}
-		}
+////////////////////////////////////////////////////////////////////////////////
+void FAnalysisMachine::Transition()
+{
+	if (ActiveStage != nullptr)
+	{
+		FMachineContext Context = { *this, Bridge };
+		ActiveStage->ExitStage(Context);
 
-		return false;
+		DeadStages.Add(ActiveStage);
 	}
 
-	switch (Header->TransportVersion)
+	ActiveStage = (StageQueue.Num() > 0) ? StageQueue.Pop() : nullptr;
+
+	if (ActiveStage != nullptr)
 	{
-	case ETransport::Raw:		Transport = new FTransport(); break;
-	case ETransport::Packet:	Transport = new FPacketTransport(); break;
-	case ETransport::TidPacket:	Transport = new FTidPacketTransport(); break;
-	default:					return false;
-	//case 'E':	/* See the magic above */ break;
-	//case 'T':	/* See the magic above */ break;
+		FMachineContext Context = { *this, Bridge };
+		ActiveStage->EnterStage(Context);
 	}
+}
 
-	ProtocolVersion = Header->ProtocolVersion;
-	switch (ProtocolVersion)
+////////////////////////////////////////////////////////////////////////////////
+FAnalysisMachine::EStatus FAnalysisMachine::OnData(FStreamReader& Reader)
+{
+	FMachineContext Context = { *this, Bridge };
+	EStatus Ret;
+	do
 	{
-	case Protocol0::EProtocol::Id:
-		ProtocolHandler = &FAnalysisEngine::OnDataProtocol0;
-		break;
-
-	case Protocol1::EProtocol::Id:
-		Serial.Mask = 0x0000ffff;
-		ProtocolHandler = &FAnalysisEngine::OnDataProtocol2;
-		break;
+		CleanUp();
+		check(ActiveStage != nullptr);
+		Ret = ActiveStage->OnData(Reader, Context);
+	}
+	while (Ret == EStatus::Continue);
+	return Ret;
+}
 
-	case Protocol2::EProtocol::Id:
-	case Protocol3::EProtocol::Id:
-	case Protocol4::EProtocol::Id:
-		Serial.Mask = 0x00ffffff;
-		ProtocolHandler = &FAnalysisEngine::OnDataProtocol2;
-		break;
 
-	default:
-		return false;
-	}
 
-	Reader.Advance(sizeof(*Header));
-	return true;
-}
+// {{{1 protocol-0 -------------------------------------------------------------
 
 ////////////////////////////////////////////////////////////////////////////////
-bool FAnalysisEngine::OnData(FStreamReader& Reader)
+class FProtocol0Stage
+	: public FAnalysisMachine::FStage
 {
-	if (Transport == nullptr)
-	{
-		// Ensure we've a reasonable amount of data to establish the transport with
-		if (Reader.GetPointer(32) == nullptr)
-		{
-			return true;
-		}
-
-		if (!EstablishTransport(Reader))
-		{
-			return false;
-		}
-
-		OnTransportEstablished();
-	}
+public:
+						FProtocol0Stage(FTransport* InTransport);
+						~FProtocol0Stage();
+	virtual EStatus		OnData(FStreamReader& Reader, const FMachineContext& Context) override;
 
-	Transport->SetReader(Reader);
-	bool bRet = (this->*ProtocolHandler)();
+private:
+	FTypeRegistry		TypeRegistry;
+	FTransport*			Transport;
+};
 
-	// If there's no analyzers left we might as well not continue
-	int32 ActiveAnalyzerCount = 0;
-	for (IAnalyzer* Analyzer : Analyzers)
-	{
-		if ((Analyzer != nullptr) && (Analyzer != this))
-		{
-			ActiveAnalyzerCount++;
-		}
-	}
+////////////////////////////////////////////////////////////////////////////////
+FProtocol0Stage::FProtocol0Stage(FTransport* InTransport)
+: Transport(InTransport)
+{
+}
 
-	return (ActiveAnalyzerCount > 0) & bRet;
+////////////////////////////////////////////////////////////////////////////////
+FProtocol0Stage::~FProtocol0Stage()
+{
+	delete Transport;
 }
 
 ////////////////////////////////////////////////////////////////////////////////
-bool FAnalysisEngine::OnDataProtocol0()
+FProtocol0Stage::EStatus FProtocol0Stage::OnData(
+	FStreamReader& Reader,
+	const FMachineContext& Context)
 {
-	FThreads::FInfo ThreadInfo;
+	Transport->SetReader(Reader);
 
 	while (true)
 	{
@@ -1442,48 +1750,99 @@ bool FAnalysisEngine::OnDataProtocol0()
 		}
 
 		uint32 Uid = uint32(Header->Uid) & uint32(Protocol0::EKnownEventUids::UidMask);
-		if (Uid >= uint32(Dispatches.Num()))
+
+		if (Uid == uint32(Protocol0::EKnownEventUids::NewEvent))
 		{
-			return false;
+			// There is no need to check size here as the runtime never builds
+			// packets that fragment new-event events.
+			const FTypeRegistry::FTypeInfo* TypeInfo = TypeRegistry.Add(Header->EventData);
+			Context.Bridge.OnNewType(TypeInfo);
+			Transport->Advance(BlockSize);
+			continue;
 		}
 
-		const FDispatch* Dispatch = Dispatches[Uid];
-		if (Dispatch == nullptr)
+		const FTypeRegistry::FTypeInfo* TypeInfo = TypeRegistry.Get(Uid);
+		if (TypeInfo == nullptr)
 		{
-			return false;
+			return EStatus::Error;
 		}
 
 		FEventDataInfo EventDataInfo = {
 			Header->EventData,
-			*Dispatch,
+			*TypeInfo,
 			nullptr,
 			Header->Size
 		};
 
-		FOnEventContext Context = {
-			(const FThreadInfo&)ThreadInfo,
-			(const FEventTime&)Timing,
-			(const FEventData&)EventDataInfo,
-		};
-
-		ForEachRoute(Dispatch->FirstRoute, false, [&] (IAnalyzer* Analyzer, uint16 RouteId)
-		{
-			if (!Analyzer->OnEvent(RouteId, EStyle::Normal, Context))
-			{
-				RetireAnalyzer(Analyzer);
-			}
-		});
+		Context.Bridge.OnEvent(EventDataInfo);
 
 		Transport->Advance(BlockSize);
 	}
 
-	return true;
+	return Reader.IsEmpty() ? EStatus::Eof : EStatus::NotEnoughData;
 }
 
+
+
+// {{{1 protocol-2 -------------------------------------------------------------
+
 ////////////////////////////////////////////////////////////////////////////////
-bool FAnalysisEngine::OnDataProtocol2()
+class FProtocol2Stage
+	: public FAnalysisMachine::FStage
+{
+public:
+								FProtocol2Stage(uint32 Version, FTransport* InTransport);
+								~FProtocol2Stage();
+	virtual EStatus 			OnData(FStreamReader& Reader, const FMachineContext& Context) override;
+	virtual void				EnterStage(const FMachineContext& Context) override;
+
+protected:
+	virtual int32				OnData(FStreamReader& Reader, FAnalysisBridge& Bridge);
+	int32						OnDataAux(FStreamReader& Reader, FAuxDataCollector& Collector);
+	FTypeRegistry				TypeRegistry;
+	FTransport*					Transport;
+	uint32						ProtocolVersion;
+};
+
+////////////////////////////////////////////////////////////////////////////////
+FProtocol2Stage::FProtocol2Stage(uint32 Version, FTransport* InTransport)
+: Transport(InTransport)
+, ProtocolVersion(Version)
+{
+}
+
+////////////////////////////////////////////////////////////////////////////////
+FProtocol2Stage::~FProtocol2Stage()
+{
+	delete Transport;
+}
+
+
+////////////////////////////////////////////////////////////////////////////////
+void FProtocol2Stage::EnterStage(const FMachineContext& Context)
+{
+	auto& Serial = Context.Bridge.GetSerial();
+
+	if (ProtocolVersion == Protocol1::EProtocol::Id)
+	{
+		Serial.Mask = 0x0000ffff;
+	}
+	else
+	{
+		/* Protocol2::EProtocol::Id
+		   Protocol3::EProtocol::Id
+		   Protocol4::EProtocol::Id */
+		Serial.Mask = 0x00ffffff;
+	}
+}
+
+////////////////////////////////////////////////////////////////////////////////
+FProtocol2Stage::EStatus FProtocol2Stage::OnData(
+	FStreamReader& Reader,
+	const FMachineContext& Context)
 {
 	auto* InnerTransport = (FTidPacketTransport*)Transport;
+	InnerTransport->SetReader(Reader);
 	InnerTransport->Update();
 
 	struct FRotaItem
@@ -1497,11 +1856,12 @@ bool FAnalysisEngine::OnDataProtocol2()
 
 	for (uint32 i = 0, n = InnerTransport->GetThreadCount(); i < n; ++i)
 	{
-		FStreamReader* Reader = InnerTransport->GetThreadStream(i);
+		FStreamReader* ThreadReader = InnerTransport->GetThreadStream(i);
 		uint32 ThreadId = InnerTransport->GetThreadId(i);
-		Rota.Add({~0u, ThreadId, Reader});
+		Rota.Add({~0u, ThreadId, ThreadReader});
 	}
 
+	auto& Serial = Context.Bridge.GetSerial();
 	uint32 ActiveCount = uint32(Rota.Num());
 	while (true)
 	{
@@ -1515,18 +1875,9 @@ bool FAnalysisEngine::OnDataProtocol2()
 				continue;
 			}
 
-			FThreads::FInfo* ThreadInfo = Threads.GetInfo(RotaItem.ThreadId);
-
-			uint32 AvailableSerial;
-			if (ProtocolVersion == UE::Trace::Protocol4::EProtocol::Id)
-			{
-				AvailableSerial = OnDataProtocol4(*(RotaItem.Reader), *ThreadInfo);
-			}
-			else
-			{
-				AvailableSerial = OnDataProtocol2(*(RotaItem.Reader), *ThreadInfo);
-			}
+			Context.Bridge.SetActiveThread(RotaItem.ThreadId);
 
+			uint32 AvailableSerial = OnData(*(RotaItem.Reader), Context.Bridge);
 			if (int32(AvailableSerial) >= 0)
 			{
 				RotaItem.Serial = AvailableSerial;
@@ -1607,12 +1958,14 @@ bool FAnalysisEngine::OnDataProtocol2()
 		}
 	}
 
-	return true;
+	return Reader.IsEmpty() ? EStatus::Eof : EStatus::NotEnoughData;
 }
 
 ////////////////////////////////////////////////////////////////////////////////
-int32 FAnalysisEngine::OnDataProtocol2(FStreamReader& Reader, FThreads::FInfo& ThreadInfo)
+int32 FProtocol2Stage::OnData(FStreamReader& Reader, FAnalysisBridge& Bridge)
 {
+	auto& Serial = Bridge.GetSerial();
+
 	while (true)
 	{
 		auto Mark = Reader.SaveMark();
@@ -1624,14 +1977,19 @@ int32 FAnalysisEngine::OnDataProtocol2(FStreamReader& Reader, FThreads::FInfo& T
 		}
 
 		uint32 Uid = uint32(Header->Uid) & uint32(Protocol2::EKnownEventUids::UidMask);
-		if (Uid >= uint32(Dispatches.Num()))
+
+		if (Uid == uint32(Protocol2::EKnownEventUids::NewEvent))
 		{
-			// We don't know about this event yet
-			break;
+			// There is no need to check size here as the runtime never builds
+			// packets that fragment new-event events.
+			const FTypeRegistry::FTypeInfo* TypeInfo = TypeRegistry.Add(Header + 1);
+			Bridge.OnNewType(TypeInfo);
+			Reader.Advance(Header->Size);
+			continue;
 		}
 
-		const FDispatch* Dispatch = Dispatches[Uid];
-		if (Dispatch == nullptr)
+		const FTypeRegistry::FTypeInfo* TypeInfo = TypeRegistry.Get(Uid);
+		if (TypeInfo == nullptr)
 		{
 			// Event-types may not to be discovered in Uid order.
 			break;
@@ -1640,7 +1998,7 @@ int32 FAnalysisEngine::OnDataProtocol2(FStreamReader& Reader, FThreads::FInfo& T
 		uint32 BlockSize = Header->Size;
 
 		// Make sure we consume events in the correct order
-		if ((Dispatch->Flags & FDispatch::Flag_NoSync) == 0)
+		if ((TypeInfo->Flags & FTypeRegistry::FTypeInfo::Flag_NoSync) == 0)
 		{
 			switch (ProtocolVersion)
 			{
@@ -1682,9 +2040,9 @@ int32 FAnalysisEngine::OnDataProtocol2(FStreamReader& Reader, FThreads::FInfo& T
 		Reader.Advance(BlockSize);
 
 		FAuxDataCollector AuxCollector;
-		if (Dispatch->Flags & FDispatch::Flag_MaybeHasAux)
+		if (TypeInfo->Flags & FTypeRegistry::FTypeInfo::Flag_MaybeHasAux)
 		{
-			int AuxStatus = OnDataProtocol2Aux(Reader, AuxCollector);
+			int AuxStatus = OnDataAux(Reader, AuxCollector);
 			if (AuxStatus == 0)
 			{
 				Reader.RestoreMark(Mark);
@@ -1692,7 +2050,7 @@ int32 FAnalysisEngine::OnDataProtocol2(FStreamReader& Reader, FThreads::FInfo& T
 			}
 		}
 
-		if ((Dispatch->Flags & FDispatch::Flag_NoSync) == 0)
+		if ((TypeInfo->Flags & FTypeRegistry::FTypeInfo::Flag_NoSync) == 0)
 		{
 			Serial.Value += 1;
 			Serial.Value &= 0x7fffffff; // don't set msb. that has other uses
@@ -1700,30 +2058,17 @@ int32 FAnalysisEngine::OnDataProtocol2(FStreamReader& Reader, FThreads::FInfo& T
 
 		FEventDataInfo EventDataInfo = {
 			(const uint8*)Header + BlockSize - Header->Size,
-			*Dispatch,
+			*TypeInfo,
 			&AuxCollector,
 			Header->Size,
 		};
-
-		FOnEventContext Context = {
-			(const FThreadInfo&)ThreadInfo,
-			(const FEventTime&)Timing,
-			(const FEventData&)EventDataInfo,
-		};
-		ForEachRoute(Dispatch->FirstRoute, false, [&] (IAnalyzer* Analyzer, uint16 RouteId)
-		{
-			if (!Analyzer->OnEvent(RouteId, EStyle::Normal, Context))
-			{
-				RetireAnalyzer(Analyzer);
-			}
-		});
 	}
 
 	return -1;
 }
 
 ////////////////////////////////////////////////////////////////////////////////
-int32 FAnalysisEngine::OnDataProtocol2Aux(FStreamReader& Reader, FAuxDataCollector& Collector)
+int32 FProtocol2Stage::OnDataAux(FStreamReader& Reader, FAuxDataCollector& Collector)
 {
 	while (true)
 	{
@@ -1766,112 +2111,49 @@ int32 FAnalysisEngine::OnDataProtocol2Aux(FStreamReader& Reader, FAuxDataCollect
 	}
 }
 
+
+
+// {{{1 protocol-4 -------------------------------------------------------------
+
 ////////////////////////////////////////////////////////////////////////////////
-int32 FAnalysisEngine::OnDataProtocol4(FStreamReader& Reader, FThreads::FInfo& ThreadInfo)
+class FProtocol4Stage
+	: public FProtocol2Stage
 {
-	while (true)
-	{
-		if (int32 TriResult = OnDataProtocol4Impl(Reader, ThreadInfo))
-		{
-			return (TriResult < 0) ? ~TriResult : -1;
-		}
-	}
-}
+public:
+					FProtocol4Stage(uint32 Version, FTransport* InTransport);
+
+private:
+	virtual int32	OnData(FStreamReader& Reader, FAnalysisBridge& Bridge) override;
+	int32			OnDataImpl(FStreamReader& Reader, FAnalysisBridge& Bridge);
+	int32			OnDataKnown(uint32 Uid, FStreamReader& Reader, FAnalysisBridge& Bridge);
+};
 
 ////////////////////////////////////////////////////////////////////////////////
-int32 FAnalysisEngine::OnDataProtocol4Known(
-	uint32 Uid,
-	FStreamReader& Reader,
-	FThreads::FInfo& ThreadInfo)
+FProtocol4Stage::FProtocol4Stage(uint32 Version, FTransport* InTransport)
+: FProtocol2Stage(Version, InTransport)
 {
-	using namespace Protocol4;
+}
 
-	switch (Uid)
+////////////////////////////////////////////////////////////////////////////////
+int32 FProtocol4Stage::OnData(FStreamReader& Reader, FAnalysisBridge& Bridge)
+{
+	while (true)
 	{
-	case EKnownEventUids::NewEvent:
-		if (const auto* Size = Reader.GetPointer<uint16>())
+		if (int32 TriResult = OnDataImpl(Reader, Bridge))
 		{
-			Reader.Advance(sizeof(*Size));
-			uint32 EventSize = *Size;
-			if (const void* EventData = Reader.GetPointer(EventSize))
-			{
-				Reader.Advance(EventSize);
-				OnNewEventInternal(EventData);
-				return 1;
-			}
-		}
-		break;
-
-	case EKnownEventUids::EnterScope:
-	case EKnownEventUids::EnterScope_T:
-		if (Uid > EKnownEventUids::EnterScope)
-		{
-			const uint8* Stamp = Reader.GetPointer(sizeof(uint64) - 1);
-			if (Stamp == nullptr)
-			{
-				break;
-			}
-			uint64 Timestamp = ThreadInfo.PrevTimestamp += *(uint64*)(Stamp - 1) >> 8;
-			ThreadInfo.ScopeRoutes.Push(~Timestamp);
-			Reader.Advance(sizeof(uint64) - 1);
-		}
-		else
-		{
-			ThreadInfo.ScopeRoutes.Push(~0);
-		}
-
-		return 1;
-
-	case EKnownEventUids::LeaveScope:
-	case EKnownEventUids::LeaveScope_T:
-		Timing.EventTimestamp = 0;
-		if (Uid > EKnownEventUids::LeaveScope)
-		{
-			const uint8* Stamp = Reader.GetPointer(sizeof(uint64) - 1);
-			if (Stamp == nullptr)
-			{
-				break;
-			}
-			Timing.EventTimestamp = ThreadInfo.PrevTimestamp += *(uint64*)(Stamp - 1) >> 8;
-			Reader.Advance(sizeof(uint64) - 1);
-		}
-
-		if (ThreadInfo.ScopeRoutes.Num() > 0)
-		{
-			uint32 RouteIndex = ThreadInfo.ScopeRoutes.Pop(false);
-
-			FDispatch EmptyDispatch = {};
-			FEventDataInfo EmptyEventInfo = { nullptr, EmptyDispatch };
-
-			FOnEventContext Context = {
-				(const FThreadInfo&)ThreadInfo,
-				(const FEventTime&)Timing,
-				(const FEventData&)EmptyEventInfo,
-			};
-			ForEachRoute(RouteIndex, true, [&] (IAnalyzer* Analyzer, uint16 RouteId)
-			{
-				if (!Analyzer->OnEvent(RouteId, EStyle::LeaveScope, Context))
-				{
-					RetireAnalyzer(Analyzer);
-				}
-			});
+			return (TriResult < 0) ? ~TriResult : -1;
 		}
-
-		Timing.EventTimestamp = 0;
-		return 1;
-	};
-
-	return 0;
+	}
 }
 
 ////////////////////////////////////////////////////////////////////////////////
-int32 FAnalysisEngine::OnDataProtocol4Impl(
-	FStreamReader& Reader,
-	FThreads::FInfo& ThreadInfo)
+int32 FProtocol4Stage::OnDataImpl(FStreamReader& Reader, FAnalysisBridge& Bridge)
 {
+	auto& Serial = Bridge.GetSerial();
+
 	/* Returns 0 if an event was successfully processed, 1 if there's not enough
 	 * data available, or ~AvailableLogSerial if the pending event is in the future */
-	using namespace Protocol4;
+	using namespace UE::Trace::Protocol4;
 
 	auto Mark = Reader.SaveMark();
 
@@ -1895,10 +2177,10 @@ int32 FAnalysisEngine::OnDataProtocol4Impl(
 	}
 	Uid >>= EKnownEventUids::_UidShift;
 
-	if (Uid < UserUidBias)
+	if (Uid < Bridge.GetUserUidBias())
 	{
 		Reader.Advance(UidBytes);
-		if (!OnDataProtocol4Known(Uid, Reader, ThreadInfo))
+		if (!OnDataKnown(Uid, Reader, Bridge))
 		{
 			Reader.RestoreMark(Mark);
 			return 1;
@@ -1907,13 +2189,8 @@ int32 FAnalysisEngine::OnDataProtocol4Impl(
 	}
 
 	// Do we know about this event type yet?
-	if (Uid >= uint32(Dispatches.Num()))
-	{
-		return 1;
-	}
-
-	const FDispatch* Dispatch = Dispatches[Uid];
-	if (Dispatch == nullptr)
+	const FTypeRegistry::FTypeInfo* TypeInfo = TypeRegistry.Get(Uid);
+	if (TypeInfo == nullptr)
 	{
 		return 1;
 	}
@@ -1928,7 +2205,7 @@ int32 FAnalysisEngine::OnDataProtocol4Impl(
 	uint32 BlockSize = Header->Size;
 
 	// Make sure we consume events in the correct order
-	if ((Dispatch->Flags & FDispatch::Flag_NoSync) == 0)
+	if ((TypeInfo->Flags & FTypeRegistry::FTypeInfo::Flag_NoSync) == 0)
 	{
 		if (Reader.GetPointer<FEventHeaderSync>() == nullptr)
 		{
@@ -1958,17 +2235,17 @@ int32 FAnalysisEngine::OnDataProtocol4Impl(
 
 	// Collect auxiliary data
 	FAuxDataCollector AuxCollector;
-	if (Dispatch->Flags & FDispatch::Flag_MaybeHasAux)
+	if (TypeInfo->Flags & FTypeRegistry::FTypeInfo::Flag_MaybeHasAux)
 	{
 		// Important events' size may include their array data so we need to backtrack
 		auto NextMark = Reader.SaveMark();
-		if (Dispatch->Flags & FDispatch::Flag_Important)
+		if (TypeInfo->Flags & FTypeRegistry::FTypeInfo::Flag_Important)
 		{
 			Reader.RestoreMark(Mark);
-			Reader.Advance(sizeof(FEventHeader) + Dispatch->EventSize);
+			Reader.Advance(sizeof(FEventHeader) + TypeInfo->EventSize);
 		}
 
-		int AuxStatus = OnDataProtocol2Aux(Reader, AuxCollector);
+		int AuxStatus = OnDataAux(Reader, AuxCollector);
 		if (AuxStatus == 0)
 		{
 			Reader.RestoreMark(Mark);
@@ -1978,14 +2255,14 @@ int32 FAnalysisEngine::OnDataProtocol4Impl(
 		// User error could have resulted in less space being used that was
 		// allocated for important events. So we can't assume that aux data
 		// reading has read all the way up to the next event. So we use marks
-		if (Dispatch->Flags & FDispatch::Flag_Important)
+		if (TypeInfo->Flags & FTypeRegistry::FTypeInfo::Flag_Important)
 		{
 			Reader.RestoreMark(NextMark);
 		}
 	}
 
 	// Maintain sync
-	if ((Dispatch->Flags & FDispatch::Flag_NoSync) == 0)
+	if ((TypeInfo->Flags & FTypeRegistry::FTypeInfo::Flag_NoSync) == 0)
 	{
 		Serial.Value += 1;
 		Serial.Value &= 0x7fffffff; // don't set msb. that has other uses
@@ -1994,39 +2271,298 @@ int32 FAnalysisEngine::OnDataProtocol4Impl(
 	// Sent the event to subscribed analyzers
 	FEventDataInfo EventDataInfo = {
 		(const uint8*)Header + BlockSize - Header->Size,
-		*Dispatch,
+		*TypeInfo,
 		&AuxCollector,
 		Header->Size,
 	};
 
-	EStyle Style = EStyle::Normal;
-	if (ThreadInfo.ScopeRoutes.Num() > 0 && int64(ThreadInfo.ScopeRoutes.Last()) < 0)
+	Bridge.OnEvent(EventDataInfo);
+
+	return 0;
+}
+
+////////////////////////////////////////////////////////////////////////////////
+int32 FProtocol4Stage::OnDataKnown(
+	uint32 Uid,
+	FStreamReader& Reader,
+	FAnalysisBridge& Bridge)
+{
+	using namespace UE::Trace::Protocol4;
+
+	switch (Uid)
 	{
-		Style = EStyle::EnterScope;
-		Timing.EventTimestamp = ~(ThreadInfo.ScopeRoutes.Last());
-		ThreadInfo.ScopeRoutes.Last() = Dispatch->FirstRoute;
+	case EKnownEventUids::NewEvent:
+		{
+			const auto* Size = Reader.GetPointer<uint16>();
+			const FTypeRegistry::FTypeInfo* TypeInfo = TypeRegistry.Add(Size + 1);
+			Bridge.OnNewType(TypeInfo);
+			Reader.Advance(sizeof(*Size) + *Size);
+			return 1;
+		}
+		break;
+
+	case EKnownEventUids::EnterScope:
+	case EKnownEventUids::EnterScope_T:
+		if (Uid > EKnownEventUids::EnterScope)
+		{
+			const uint8* Stamp = Reader.GetPointer(sizeof(uint64) - 1);
+			if (Stamp == nullptr)
+			{
+				break;
+			}
+
+			uint64 Timestamp = *(uint64*)(Stamp - 1) >> 8;
+			Bridge.EnterScope(Timestamp);
+
+			Reader.Advance(sizeof(uint64) - 1);
+		}
+		else
+		{
+			Bridge.EnterScope();
+		}
+		return 1;
+
+	case EKnownEventUids::LeaveScope:
+	case EKnownEventUids::LeaveScope_T:
+		if (Uid > EKnownEventUids::LeaveScope)
+		{
+			const uint8* Stamp = Reader.GetPointer(sizeof(uint64) - 1);
+			if (Stamp == nullptr)
+			{
+				break;
+			}
+
+			uint64 Timestamp = *(uint64*)(Stamp - 1) >> 8;
+			Bridge.LeaveScope(Timestamp);
+
+			Reader.Advance(sizeof(uint64) - 1);
+		}
+		else
+		{
+			Bridge.LeaveScope();
+		}
+		return 1;
+	};
+
+	return 0;
+}
+
+
+
+// {{{1 est.-transport ---------------------------------------------------------
+
+////////////////////////////////////////////////////////////////////////////////
+class FEstablishTransportStage
+	: public FAnalysisMachine::FStage
+{
+public:
+	virtual EStatus OnData(FStreamReader& Reader, const FMachineContext& Context) override;
+};
+
+////////////////////////////////////////////////////////////////////////////////
+FEstablishTransportStage::EStatus FEstablishTransportStage::OnData(
+	FStreamReader& Reader,
+	const FMachineContext& Context)
+{
+	using namespace UE::Trace;
+
+	const struct {
+		uint8 TransportVersion;
+		uint8 ProtocolVersion;
+	}* Header = decltype(Header)(Reader.GetPointer(sizeof(*Header)));
+	if (Header == nullptr)
+	{
+		return EStatus::NotEnoughData;
 	}
-	else
+
+	FTransport* Transport = nullptr;
+	switch (Header->TransportVersion)
 	{
-		Timing.EventTimestamp = 0;
+	case ETransport::Raw:		Transport = new FTransport(); break;
+	case ETransport::Packet:	Transport = new FPacketTransport(); break;
+	case ETransport::TidPacket:	Transport = new FTidPacketTransport(); break;
+	default:					return EStatus::Error;
 	}
 
-	FOnEventContext Context = {
-		(const FThreadInfo&)ThreadInfo,
-		(const FEventTime&)Timing,
-		(const FEventData&)EventDataInfo,
-	};
+	uint32 ProtocolVersion = Header->ProtocolVersion;
+	switch (ProtocolVersion)
+	{
+	case Protocol0::EProtocol::Id:
+		Context.Machine.QueueStage<FProtocol0Stage>(Transport);
+		Context.Machine.Transition();
+		break;
+
+	case Protocol1::EProtocol::Id:
+		Context.Machine.QueueStage<FProtocol2Stage>(ProtocolVersion, Transport);
+		Context.Machine.Transition();
+		break;
+
+	case Protocol2::EProtocol::Id:
+	case Protocol3::EProtocol::Id:
+	case Protocol4::EProtocol::Id:
+		Context.Machine.QueueStage<FProtocol4Stage>(ProtocolVersion, Transport);
+		Context.Machine.Transition();
+		break;
+
+	default:
+		return EStatus::Error;
+	}
+
+	Reader.Advance(sizeof(*Header));
+	return EStatus::Continue;
+}
+
+
+
+// {{{1 metadata ---------------------------------------------------------------
+
+////////////////////////////////////////////////////////////////////////////////
+class FMetadataStage
+	: public FAnalysisMachine::FStage
+{
+public:
+	virtual EStatus OnData(FStreamReader& Reader, const FMachineContext& Context) override;
+};
 
-	bool bScoped = (Style != EStyle::Normal);
-	ForEachRoute(Dispatch->FirstRoute, bScoped, [&] (IAnalyzer* Analyzer, uint16 RouteId)
+////////////////////////////////////////////////////////////////////////////////
+FMetadataStage::EStatus FMetadataStage::OnData(
+	FStreamReader& Reader,
+	const FMachineContext& Context)
+{
+	const auto* MetadataSize = Reader.GetPointer<uint16>();
+	if (MetadataSize == nullptr)
 	{
-		if (!Analyzer->OnEvent(RouteId, Style, Context))
-		{
-			RetireAnalyzer(Analyzer);
-		}
-	});
+		return EStatus::NotEnoughData;
+	}
 
-	return 0;
+	Reader.Advance(sizeof(*MetadataSize) + *MetadataSize);
+
+	Context.Machine.QueueStage<FEstablishTransportStage>();
+	Context.Machine.Transition();
+	return EStatus::Continue;
+}
+
+
+
+// {{{1 magic ------------------------------------------------------------------
+
+////////////////////////////////////////////////////////////////////////////////
+class FMagicStage
+	: public FAnalysisMachine::FStage
+{
+public:
+	virtual EStatus OnData(FStreamReader& Reader, const FMachineContext& Context) override;
+};
+
+////////////////////////////////////////////////////////////////////////////////
+FMagicStage::EStatus FMagicStage::OnData(
+	FStreamReader& Reader,
+	const FMachineContext& Context)
+{
+	const auto* MagicPtr = Reader.GetPointer<uint32>();
+	if (MagicPtr == nullptr)
+	{
+		return EStatus::NotEnoughData;
+	}
+
+	uint32 Magic = *MagicPtr;
+	Reader.Advance(sizeof(*MagicPtr));
+
+	if (Magic == 'ECRT' || Magic == '2CRT')
+	{
+		// Source is big-endian which we don't currently support
+		return EStatus::Error;
+	}
+
+	if (Magic == 'TRCE')
+	{
+		Context.Machine.QueueStage<FEstablishTransportStage>();
+		Context.Machine.Transition();
+		return EStatus::Continue;
+	}
+
+	if (Magic == 'TRC2')
+	{
+		Context.Machine.QueueStage<FMetadataStage>();
+		Context.Machine.Transition();
+		return EStatus::Continue;
+	}
+
+	return EStatus::Error;
+}
+
+
+
+// {{{1 engine -----------------------------------------------------------------
+
+////////////////////////////////////////////////////////////////////////////////
+class FAnalysisEngine::FImpl
+{
+public:
+						FImpl(TArray<IAnalyzer*>&& Analyzers);
+	void				Begin();
+	void				End();
+	bool				OnData(FStreamReader& Reader);
+	FAnalysisBridge		Bridge;
+	FAnalysisMachine	Machine = Bridge;
+};
+
+////////////////////////////////////////////////////////////////////////////////
+FAnalysisEngine::FImpl::FImpl(TArray<IAnalyzer*>&& Analyzers)
+: Bridge(Forward<TArray<IAnalyzer*>>(Analyzers))
+{
+}
+
+////////////////////////////////////////////////////////////////////////////////
+void FAnalysisEngine::FImpl::Begin()
+{
+	Machine.QueueStage<FMagicStage>();
+	Machine.Transition();
+}
+
+////////////////////////////////////////////////////////////////////////////////
+void FAnalysisEngine::FImpl::End()
+{
+	Machine.Transition();
+	Bridge.Reset();
+}
+
+////////////////////////////////////////////////////////////////////////////////
+bool FAnalysisEngine::FImpl::OnData(FStreamReader& Reader)
+{
+	return (Machine.OnData(Reader) != FAnalysisMachine::EStatus::Error);
+}
+
+
+
+////////////////////////////////////////////////////////////////////////////////
+FAnalysisEngine::FAnalysisEngine(TArray<IAnalyzer*>&& Analyzers)
+: Impl(new FImpl(Forward<TArray<IAnalyzer*>>(Analyzers)))
+{
+}
+
+////////////////////////////////////////////////////////////////////////////////
+FAnalysisEngine::~FAnalysisEngine()
+{
+	delete Impl;
+}
+
+////////////////////////////////////////////////////////////////////////////////
+void FAnalysisEngine::Begin()
+{
+	Impl->Begin();
+}
+
+////////////////////////////////////////////////////////////////////////////////
+void FAnalysisEngine::End()
+{
+	Impl->End();
+}
+
+////////////////////////////////////////////////////////////////////////////////
+bool FAnalysisEngine::OnData(FStreamReader& Reader)
+{
+	return Impl->OnData(Reader);
 }
 
   // }}}
diff --git a/Engine/Source/Developer/TraceAnalysis/Private/Analysis/Engine.h b/Engine/Source/Developer/TraceAnalysis/Private/Analysis/Engine.h
--- a/Engine/Source/Developer/TraceAnalysis/Private/Analysis/Engine.h
+++ b/Engine/Source/Developer/TraceAnalysis/Private/Analysis/Engine.h
@@ -2,130 +2,31 @@
 
 #pragma once
 
-#include "Containers/Array.h"
-#include "Trace/Analysis.h"
-#include "Trace/Analyzer.h"
+#include "Containers/ContainersFwd.h"
 
 namespace UE {
 namespace Trace {
 
-struct FAuxDataCollector;
-struct FDispatch;
-class FDispatchBuilder;
+class IAnalyzer;
 class FStreamReader;
 
-////////////////////////////////////////////////////////////////////////////////
-class FThreads
-{
-public:
-	struct FInfo
-	{
-		TArray<uint64>	ScopeRoutes;
-		TArray<uint8>	Name;
-		TArray<uint8>	GroupName;
-		int64			PrevTimestamp = 0;
-		uint32			ThreadId = ~0u;
-		uint32			SystemId = 0;
-		int32			SortHint = 0xff;
-	};
-
-						FThreads();
-						~FThreads();
-	FInfo*				GetInfo();
-	FInfo*				GetInfo(uint32 ThreadId);
-	void				SetGroupName(const ANSICHAR* InGroupName, uint32 Length);
-	const TArray<uint8>*GetGroupName() const;
-
-private:
-	uint32				LastGetInfoId = ~0u;
-	TArray<FInfo*>		Infos;
-	TArray<uint8>		GroupName;
-};
-
-
-
-////////////////////////////////////////////////////////////////////////////////
-struct FTiming
-{
-	uint64	BaseTimestamp;
-	uint64	TimestampHz;
-	double	InvTimestampHz;
-	uint64	EventTimestamp = 0;
-};
-
-
-
 ////////////////////////////////////////////////////////////////////////////////
 class FAnalysisEngine
-	: public IAnalyzer
 {
 public:
-						FAnalysisEngine(TArray<IAnalyzer*>&& InAnalyzers);
+						FAnalysisEngine(TArray<IAnalyzer*>&& Analyzers);
 						~FAnalysisEngine();
-	bool				OnData(FStreamReader& Reader);
 	void				Begin();
 	void				End();
+	bool				OnData(FStreamReader& Reader);
 
 private:
-	typedef bool (FAnalysisEngine::*ProtocolHandlerType)();
-
-	struct FRoute
-	{
-		uint32			Hash;
-		uint16			AnalyzerIndex : 15;
-		uint16			bScoped : 1;
-		uint16			Id;
-		union
-		{
-			uint32		ParentHash;
-			struct
-			{
-				int16	Count;
-				int16	Parent;
-			};
-		};
-	};
-
-	void				OnTransportEstablished();
-	virtual bool		OnEvent(uint16 RouteId, EStyle Style, const FOnEventContext& Context) override;
-	virtual void		OnAnalysisBegin(const FOnAnalysisContext& Context) override;
-	void				OnNewTrace(const FOnEventContext& Context);
-	void				OnSerialSync(const FOnEventContext& Context);
-	void				OnThreadTiming(const FOnEventContext& Context);
-	void				OnThreadInfoInternal(const FOnEventContext& Context);
-	void				OnThreadGroupBegin(const FOnEventContext& Context);
-	void				OnThreadGroupEnd();
-	void				OnTiming(const FOnEventContext& Context);
-	void				OnNewEventInternal(const void* EventData);
-	void				OnNewEventProtocol0(FDispatchBuilder& Builder, const void* EventData);
-	void				OnNewEventProtocol1(FDispatchBuilder& Builder, const void* EventData);
-	bool				EstablishTransport(FStreamReader& Reader);
-	bool				OnDataProtocol0();
-	bool				OnDataProtocol2();
-	int32				OnDataProtocol2(FStreamReader& Reader, FThreads::FInfo& Info);
-	int32				OnDataProtocol2Aux(FStreamReader& Reader, FAuxDataCollector& Collector);
-	int32				OnDataProtocol4(FStreamReader& Reader, FThreads::FInfo& Info);
-	int32				OnDataProtocol4Impl(FStreamReader& Reader, FThreads::FInfo& ThreadInfo);
-	int32				OnDataProtocol4Known(uint32 Uid, FStreamReader& Reader, FThreads::FInfo& ThreadInfo);
-	bool				AddDispatch(FDispatch* Dispatch);
-	template <typename ImplType>
-	void				ForEachRoute(uint32 RouteIndex, bool bScoped, ImplType&& Impl);
-	void				AddRoute(uint16 AnalyzerIndex, uint16 Id, const ANSICHAR* Logger, const ANSICHAR* Event, bool bScoped);
-	void				RetireAnalyzer(IAnalyzer* Analyzer);
-	FThreads			Threads;
-	FTiming				Timing;
-	TArray<FRoute>		Routes;
-	TArray<IAnalyzer*>	Analyzers;
-	TArray<FDispatch*>	Dispatches;
-	class FTransport*	Transport = nullptr;
-	ProtocolHandlerType	ProtocolHandler = nullptr;
-	struct
-	{
-		uint32			Value = 0;
-		uint32			Mask = 0;
-	}					Serial;
-	uint32				UserUidBias = 1; // 1 because new-event events must exists be uid zero.
-	uint8				ProtocolVersion = 0;
+	class FImpl;
+						FAnalysisEngine(const FAnalysisEngine&) = delete;
+						FAnalysisEngine(const FAnalysisEngine&&) = delete;
+	FAnalysisEngine		operator = (const FAnalysisEngine&) = delete;
+	FAnalysisEngine		operator = (const FAnalysisEngine&&) = delete;
+	FImpl*				Impl;
 };
 
 } // namespace Trace
diff --git a/Engine/Source/Developer/TraceAnalysis/Private/Analysis/Processor.cpp b/Engine/Source/Developer/TraceAnalysis/Private/Analysis/Processor.cpp
--- a/Engine/Source/Developer/TraceAnalysis/Private/Analysis/Processor.cpp
+++ b/Engine/Source/Developer/TraceAnalysis/Private/Analysis/Processor.cpp
@@ -6,7 +6,6 @@
 #include "HAL/RunnableThread.h"
 #include "StreamReader.h"
 #include "Templates/UnrealTemplate.h"
-#include "Trace/Analysis.h"
 #include "Trace/DataStream.h"
 
 namespace UE {
diff --git a/Engine/Source/Developer/TraceAnalysis/Private/Analysis/Processor.h b/Engine/Source/Developer/TraceAnalysis/Private/Analysis/Processor.h
--- a/Engine/Source/Developer/TraceAnalysis/Private/Analysis/Processor.h
+++ b/Engine/Source/Developer/TraceAnalysis/Private/Analysis/Processor.h
@@ -4,6 +4,7 @@
 
 #include "Engine.h"
 #include "HAL/Runnable.h"
+#include "Trace/Analysis.h"
 
 class FEvent;
 class FRunnableThread;
