Added a method to return the number of bytes remaining in a stream of trace data

#rnx

diff --git a/Engine/Source/Developer/TraceAnalysis/Private/Analysis/StreamReader.cpp b/Engine/Source/Developer/TraceAnalysis/Private/Analysis/StreamReader.cpp
--- a/Engine/Source/Developer/TraceAnalysis/Private/Analysis/StreamReader.cpp
+++ b/Engine/Source/Developer/TraceAnalysis/Private/Analysis/StreamReader.cpp
@@ -31,6 +31,12 @@ void FStreamReader::Advance(uint32 Size)
 	Cursor += Size;
 }
 
+////////////////////////////////////////////////////////////////////////////////
+int32 FStreamReader::GetRemaining() const
+{
+	return int32(PTRINT(End - Cursor));
+}
+
 ////////////////////////////////////////////////////////////////////////////////
 bool FStreamReader::IsEmpty() const
 {
diff --git a/Engine/Source/Developer/TraceAnalysis/Private/Analysis/StreamReader.h b/Engine/Source/Developer/TraceAnalysis/Private/Analysis/StreamReader.h
--- a/Engine/Source/Developer/TraceAnalysis/Private/Analysis/StreamReader.h
+++ b/Engine/Source/Developer/TraceAnalysis/Private/Analysis/StreamReader.h
@@ -17,6 +17,7 @@ public:
 	const uint8*				GetPointer(uint32 Size);
 	void						Advance(uint32 Size);
 	bool						IsEmpty() const;
+	int32						GetRemaining() const;
 	struct FMark*				SaveMark() const;
 	void						RestoreMark(struct FMark* Mark);
 
