Unchecked versions of methods to get a pointer to the data in a trace stream. While redundant on the surface, they will help to indicate when there is an expectation that enough data is known to be available.

#rnx

diff --git a/Engine/Source/Developer/TraceAnalysis/Private/Analysis/StreamReader.h b/Engine/Source/Developer/TraceAnalysis/Private/Analysis/StreamReader.h
--- a/Engine/Source/Developer/TraceAnalysis/Private/Analysis/StreamReader.h
+++ b/Engine/Source/Developer/TraceAnalysis/Private/Analysis/StreamReader.h
@@ -14,7 +14,10 @@ public:
 								~FStreamReader();
 	template <typename Type>
 	Type const*					GetPointer();
+	template <typename Type>
+	Type const*					GetPointerUnchecked();
 	const uint8*				GetPointer(uint32 Size);
+	const uint8*				GetPointerUnchecked();
 	void						Advance(uint32 Size);
 	bool						IsEmpty() const;
 	int32						GetRemaining() const;
@@ -35,6 +38,19 @@ Type const* FStreamReader::GetPointer()
 	return (Type const*)GetPointer(sizeof(Type));
 }
 
+////////////////////////////////////////////////////////////////////////////////
+template <typename Type>
+Type const* FStreamReader::GetPointerUnchecked()
+{
+	return (Type const*)GetPointerUnchecked();
+}
+
+////////////////////////////////////////////////////////////////////////////////
+inline const uint8* FStreamReader::GetPointerUnchecked()
+{
+	return Buffer + Cursor;
+}
+
 
 
 ////////////////////////////////////////////////////////////////////////////////
