Allow externally-construct trace event types to be added into the type registry.

#rnx

diff --git a/Engine/Source/Developer/TraceAnalysis/Private/Analysis/Engine.cpp b/Engine/Source/Developer/TraceAnalysis/Private/Analysis/Engine.cpp
--- a/Engine/Source/Developer/TraceAnalysis/Private/Analysis/Engine.cpp
+++ b/Engine/Source/Developer/TraceAnalysis/Private/Analysis/Engine.cpp
@@ -805,6 +805,7 @@ public:
 
 						~FTypeRegistry();
 	const FTypeInfo*	Add(const void* TraceData);
+	void				Add(FTypeInfo* TypeInfo);
 	const FTypeInfo*	Get(uint32 Uid) const;
 
 private:
@@ -884,7 +885,14 @@ const FTypeRegistry::FTypeInfo* FTypeRegistry::Add(const void* TraceData)
 	}
 
 	FTypeInfo* TypeInfo = Builder.Finalize();
+	Add(TypeInfo);
 
+	return TypeInfo;
+}
+
+////////////////////////////////////////////////////////////////////////////////
+void FTypeRegistry::Add(FTypeInfo* TypeInfo)
+{
 	// Add the type to the type-infos table. Usually duplicates are an error
 	// but due to backwards compatibility we'll override existing types.
 	uint16 Uid = TypeInfo->Uid;
@@ -901,7 +909,7 @@ const FTypeRegistry::FTypeInfo* FTypeRegistry::Add(const void* TraceData)
 		TypeInfos.SetNum(Uid + 1);
  	}
 
-	return TypeInfos[Uid] = TypeInfo;
+	TypeInfos[Uid] = TypeInfo;
 }
 
 // {{{1 analyzer-hub -----------------------------------------------------------
