Fixed segfault when analysing old protocol-0 traces.

#rnx

diff --git a/Engine/Source/Developer/TraceAnalysis/Private/Analysis/Engine.cpp b/Engine/Source/Developer/TraceAnalysis/Private/Analysis/Engine.cpp
--- a/Engine/Source/Developer/TraceAnalysis/Private/Analysis/Engine.cpp
+++ b/Engine/Source/Developer/TraceAnalysis/Private/Analysis/Engine.cpp
@@ -1717,6 +1717,7 @@ public:
 						FProtocol0Stage(FTransport* InTransport);
 						~FProtocol0Stage();
 	virtual EStatus		OnData(FStreamReader& Reader, const FMachineContext& Context) override;
+	virtual void		EnterStage(const FMachineContext& Context) override;
 
 private:
 	FTypeRegistry		TypeRegistry;
@@ -1735,6 +1736,12 @@ FProtocol0Stage::~FProtocol0Stage()
 	delete Transport;
 }
 
+////////////////////////////////////////////////////////////////////////////////
+void FProtocol0Stage::EnterStage(const FMachineContext& Context)
+{
+	Context.Bridge.SetActiveThread(0);
+}
+
 ////////////////////////////////////////////////////////////////////////////////
 FProtocol0Stage::EStatus FProtocol0Stage::OnData(
 	FStreamReader& Reader,
