Ancient traces have no magic FourCC. We can say though that they definitely start with [REDACTED_PHONE]. So we'll use that as a FourCC to identfiy antique traces.
#rnx

diff --git a/Engine/Source/Developer/TraceAnalysis/Private/Analysis/Engine.cpp b/Engine/Source/Developer/TraceAnalysis/Private/Analysis/Engine.cpp
--- a/Engine/Source/Developer/TraceAnalysis/Private/Analysis/Engine.cpp
+++ b/Engine/Source/Developer/TraceAnalysis/Private/Analysis/Engine.cpp
@@ -2481,7 +2481,6 @@ FMagicStage::EStatus FMagicStage::OnData(
 	}
 
 	uint32 Magic = *MagicPtr;
-	Reader.Advance(sizeof(*MagicPtr));
 
 	if (Magic == 'ECRT' || Magic == '2CRT')
 	{
@@ -2491,6 +2490,7 @@ FMagicStage::EStatus FMagicStage::OnData(
 
 	if (Magic == 'TRCE')
 	{
+		Reader.Advance(sizeof(*MagicPtr));
 		Context.Machine.QueueStage<FEstablishTransportStage>();
 		Context.Machine.Transition();
 		return EStatus::Continue;
@@ -2498,11 +2498,21 @@ FMagicStage::EStatus FMagicStage::OnData(
 
 	if (Magic == 'TRC2')
 	{
+		Reader.Advance(sizeof(*MagicPtr));
 		Context.Machine.QueueStage<FMetadataStage>();
 		Context.Machine.Transition();
 		return EStatus::Continue;
 	}
 
+	// There was no header on early traces so they went straight into declaring
+	// protocol and transport versions.
+	if (Magic == 0x00'00'00'01) // protocol 0, transport 1
+	{
+		Context.Machine.QueueStage<FEstablishTransportStage>();
+		Context.Machine.Transition();
+		return EStatus::Continue;
+	}
+
 	return EStatus::Error;
 }
 
