Older trace protocols had synchronised new-event events so we can't short-circuit those events.

#rnx

diff --git a/Engine/Source/Developer/TraceAnalysis/Private/Analysis/Engine.cpp b/Engine/Source/Developer/TraceAnalysis/Private/Analysis/Engine.cpp
--- a/Engine/Source/Developer/TraceAnalysis/Private/Analysis/Engine.cpp
+++ b/Engine/Source/Developer/TraceAnalysis/Private/Analysis/Engine.cpp
@@ -1824,6 +1824,31 @@ FProtocol2Stage::FProtocol2Stage(uint32 Version, FTransport* InTransport)
 : Transport(InTransport)
 , ProtocolVersion(Version)
 {
+	switch (ProtocolVersion)
+	{
+	case Protocol0::EProtocol::Id:
+	case Protocol1::EProtocol::Id:
+	case Protocol2::EProtocol::Id:
+		{
+			FDispatchBuilder Dispatch;
+			Dispatch.SetUid(uint16(Protocol2::EKnownEventUids::NewEvent));
+			Dispatch.SetLoggerName("$Trace");
+			Dispatch.SetEventName("NewEvent");
+			TypeRegistry.Add(Dispatch.Finalize());
+		}
+		break;
+
+	case Protocol3::EProtocol::Id:
+		{
+			FDispatchBuilder Dispatch;
+			Dispatch.SetUid(uint16(Protocol3::EKnownEventUids::NewEvent));
+			Dispatch.SetLoggerName("$Trace");
+			Dispatch.SetEventName("NewEvent");
+			Dispatch.SetNoSync();
+			TypeRegistry.Add(Dispatch.Finalize());
+		}
+		break;
+	}
 }
 
 ////////////////////////////////////////////////////////////////////////////////
@@ -1993,16 +2018,6 @@ int32 FProtocol2Stage::OnData(FStreamReader& Reader, FAnalysisBridge& Bridge)
 
 		uint32 Uid = uint32(Header->Uid) & uint32(Protocol2::EKnownEventUids::UidMask);
 
-		if (Uid == uint32(Protocol2::EKnownEventUids::NewEvent))
-		{
-			// There is no need to check size here as the runtime never builds
-			// packets that fragment new-event events.
-			const FTypeRegistry::FTypeInfo* TypeInfo = TypeRegistry.Add(Header + 1);
-			Bridge.OnNewType(TypeInfo);
-			Reader.Advance(Header->Size);
-			continue;
-		}
-
 		const FTypeRegistry::FTypeInfo* TypeInfo = TypeRegistry.Get(Uid);
 		if (TypeInfo == nullptr)
 		{
@@ -2071,12 +2086,25 @@ int32 FProtocol2Stage::OnData(FStreamReader& Reader, FAnalysisBridge& Bridge)
 			Serial.Value &= 0x7fffffff; // don't set msb. that has other uses
 		}
 
-		FEventDataInfo EventDataInfo = {
-			(const uint8*)Header + BlockSize - Header->Size,
-			*TypeInfo,
-			&AuxCollector,
-			Header->Size,
-		};
+		auto* EventData = (const uint8*)Header + BlockSize - Header->Size;
+		if (Uid == uint32(Protocol2::EKnownEventUids::NewEvent))
+		{
+			// There is no need to check size here as the runtime never builds
+			// packets that fragment new-event events.
+			TypeInfo = TypeRegistry.Add(EventData);
+			Bridge.OnNewType(TypeInfo);
+		}
+		else
+		{
+			FEventDataInfo EventDataInfo = {
+				EventData,
+				*TypeInfo,
+				&AuxCollector,
+				Header->Size,
+			};
+
+			Bridge.OnEvent(EventDataInfo);
+		}
 	}
 
 	return -1;
