Removed a redundant comment.

#rnx

diff --git a/Engine/Source/Developer/TraceAnalysis/Private/Analysis/Engine.cpp b/Engine/Source/Developer/TraceAnalysis/Private/Analysis/Engine.cpp
--- a/Engine/Source/Developer/TraceAnalysis/Private/Analysis/Engine.cpp
+++ b/Engine/Source/Developer/TraceAnalysis/Private/Analysis/Engine.cpp
@@ -1869,9 +1869,6 @@ void FProtocol2Stage::EnterStage(const FMachineContext& Context)
 	}
 	else
 	{
-		/* Protocol2::EProtocol::Id
-		   Protocol3::EProtocol::Id
-		   Protocol4::EProtocol::Id */
 		Serial.Mask = 0x00ffffff;
 	}
 }
