Flipped GetThreadId()'s return type to be unsigned.

#rnx

diff --git a/Engine/Source/Developer/TraceAnalysis/Private/Analysis/Transport/TidPacketTransport.cpp b/Engine/Source/Developer/TraceAnalysis/Private/Analysis/Transport/TidPacketTransport.cpp
--- a/Engine/Source/Developer/TraceAnalysis/Private/Analysis/Transport/TidPacketTransport.cpp
+++ b/Engine/Source/Developer/TraceAnalysis/Private/Analysis/Transport/TidPacketTransport.cpp
@@ -113,7 +113,7 @@ FStreamReader* FTidPacketTransport::GetThreadStream(uint32 Index)
 }
 
 ////////////////////////////////////////////////////////////////////////////////
-int32 FTidPacketTransport::GetThreadId(uint32 Index) const
+uint32 FTidPacketTransport::GetThreadId(uint32 Index) const
 {
 	return Threads[Index].ThreadId;
 }
diff --git a/Engine/Source/Developer/TraceAnalysis/Private/Analysis/Transport/TidPacketTransport.h b/Engine/Source/Developer/TraceAnalysis/Private/Analysis/Transport/TidPacketTransport.h
--- a/Engine/Source/Developer/TraceAnalysis/Private/Analysis/Transport/TidPacketTransport.h
+++ b/Engine/Source/Developer/TraceAnalysis/Private/Analysis/Transport/TidPacketTransport.h
@@ -22,7 +22,7 @@ public:
 	void					Update();
 	uint32					GetThreadCount() const;
 	FStreamReader*			GetThreadStream(uint32 Index);
-	int32					GetThreadId(uint32 Index) const;
+	uint32					GetThreadId(uint32 Index) const;
 
 private:
 	struct FThreadStream
