A case of copy-paste caused analyzers to not receive the scoped trace events they had subscribed to.

diff --git a/Engine/Source/Developer/TraceAnalysis/Private/Analysis/Engine.cpp b/Engine/Source/Developer/TraceAnalysis/Private/Analysis/Engine.cpp
--- a/Engine/Source/Developer/TraceAnalysis/Private/Analysis/Engine.cpp
+++ b/Engine/Source/Developer/TraceAnalysis/Private/Analysis/Engine.cpp
@@ -1071,7 +1071,8 @@ void FAnalyzerHub::OnEvent(
 		return;
 	}
 
-	ForEachRoute(RouteIndex, false, [&] (IAnalyzer* Analyzer, uint16 RouteId)
+	bool bScoped = (Style != IAnalyzer::EStyle::Normal);
+	ForEachRoute(RouteIndex, bScoped, [&] (IAnalyzer* Analyzer, uint16 RouteId)
 	{
 		if (!Analyzer->OnEvent(RouteId, Style, Context))
 		{
