A little tidy around the updating a thread's event scope book keeping.

diff --git a/Engine/Source/Developer/TraceAnalysis/Private/Analysis/Engine.cpp b/Engine/Source/Developer/TraceAnalysis/Private/Analysis/Engine.cpp
--- a/Engine/Source/Developer/TraceAnalysis/Private/Analysis/Engine.cpp
+++ b/Engine/Source/Developer/TraceAnalysis/Private/Analysis/Engine.cpp
@@ -1511,31 +1511,27 @@ void FAnalysisBridge::OnNewType(const FTypeRegistry::FTypeInfo* TypeInfo)
 ////////////////////////////////////////////////////////////////////////////////
 void FAnalysisBridge::OnEvent(const FEventDataInfo& EventDataInfo)
 {
+	// TODO: "Dispatch" should be renamed "EventTypeInfo" or similar.
+	const FTypeRegistry::FTypeInfo* TypeInfo = &(EventDataInfo.Dispatch);
+
 	IAnalyzer::EStyle Style = IAnalyzer::EStyle::Normal;
 	if (ThreadInfo->ScopeRoutes.Num() > 0 && int64(ThreadInfo->ScopeRoutes.Last()) < 0)
 	{
 		Style = IAnalyzer::EStyle::EnterScope;
 		State.Timing.EventTimestamp = ~(ThreadInfo->ScopeRoutes.Last());
+		ThreadInfo->ScopeRoutes.Last() = PTRINT(TypeInfo);
 	}
 	else
 	{
 		State.Timing.EventTimestamp = 0;
 	}
 
-	// TODO "Dispatch" should be renamed "EventTypeInfo" or similar.
-	const FTypeRegistry::FTypeInfo* TypeInfo = &(EventDataInfo.Dispatch);
-
 	IAnalyzer::FOnEventContext Context = {
 		*(const IAnalyzer::FThreadInfo*)ThreadInfo,
 		(const IAnalyzer::FEventTime&)(State.Timing),
 		(const IAnalyzer::FEventData&)EventDataInfo,
 	};
 	AnalyzerHub.OnEvent(*TypeInfo, Style, Context);
-
-	if (Style == IAnalyzer::EStyle::EnterScope)
-	{
-		ThreadInfo->ScopeRoutes.Last() = PTRINT(TypeInfo);
-	}
 }
 
 ////////////////////////////////////////////////////////////////////////////////
