Resolve any ambiguity between the different versions of the trace protocols' structs.

diff --git a/Engine/Source/Developer/TraceAnalysis/Private/Analysis/Engine.cpp b/Engine/Source/Developer/TraceAnalysis/Private/Analysis/Engine.cpp
--- a/Engine/Source/Developer/TraceAnalysis/Private/Analysis/Engine.cpp
+++ b/Engine/Source/Developer/TraceAnalysis/Private/Analysis/Engine.cpp
@@ -2190,7 +2190,6 @@ int32 FProtocol4Stage::OnDataImpl(FStreamReader& Reader, FAnalysisBridge& Bridge
 
 	/* Returns 0 if an event was successfully processed, 1 if there's not enough
 	 * data available, or ~AvailableLogSerial if the pending event is in the future */
-	using namespace UE::Trace::Protocol4;
 
 	auto Mark = Reader.SaveMark();
 
@@ -2200,7 +2199,7 @@ int32 FProtocol4Stage::OnDataImpl(FStreamReader& Reader, FAnalysisBridge& Bridge
 		return 1;
 	}
 
-	uint32 UidBytes = 1 + !!(*UidCursor & EKnownEventUids::Flag_TwoByteUid);
+	uint32 UidBytes = 1 + !!(*UidCursor & Protocol4::EKnownEventUids::Flag_TwoByteUid);
 	if (UidBytes > 1 && Reader.GetPointer(UidBytes) == nullptr)
 	{
 		return 1;
@@ -2212,7 +2211,7 @@ int32 FProtocol4Stage::OnDataImpl(FStreamReader& Reader, FAnalysisBridge& Bridge
 		case 1:	Uid = *UidCursor;			break;
 		case 2:	Uid = *(uint16*)UidCursor;	break;
 	}
-	Uid >>= EKnownEventUids::_UidShift;
+	Uid >>= Protocol4::EKnownEventUids::_UidShift;
 
 	if (Uid < Bridge.GetUserUidBias())
 	{
@@ -2233,7 +2232,7 @@ int32 FProtocol4Stage::OnDataImpl(FStreamReader& Reader, FAnalysisBridge& Bridge
 	}
 
 	// Parse the header
-	const auto* Header = Reader.GetPointer<FEventHeader>();
+	const auto* Header = Reader.GetPointer<Protocol4::FEventHeader>();
 	if (Header == nullptr)
 	{
 		return 1;
@@ -2244,12 +2243,12 @@ int32 FProtocol4Stage::OnDataImpl(FStreamReader& Reader, FAnalysisBridge& Bridge
 	// Make sure we consume events in the correct order
 	if ((TypeInfo->Flags & FTypeRegistry::FTypeInfo::Flag_NoSync) == 0)
 	{
-		if (Reader.GetPointer<FEventHeaderSync>() == nullptr)
+		if (Reader.GetPointer<Protocol4::FEventHeaderSync>() == nullptr)
 		{
 			return 1;
 		}
 		
-		const auto* HeaderSync = (FEventHeaderSync*)Header;
+		const auto* HeaderSync = (Protocol4::FEventHeaderSync*)Header;
 		uint32 EventSerial = HeaderSync->SerialLow|(uint32(HeaderSync->SerialHigh) << 16);
 		if (EventSerial != (Serial.Value & Serial.Mask))
 		{
@@ -2279,7 +2278,7 @@ int32 FProtocol4Stage::OnDataImpl(FStreamReader& Reader, FAnalysisBridge& Bridge
 		if (TypeInfo->Flags & FTypeRegistry::FTypeInfo::Flag_Important)
 		{
 			Reader.RestoreMark(Mark);
-			Reader.Advance(sizeof(FEventHeader) + TypeInfo->EventSize);
+			Reader.Advance(sizeof(Protocol4::FEventHeader) + TypeInfo->EventSize);
 		}
 
 		int AuxStatus = OnDataAux(Reader, AuxCollector);
@@ -2324,11 +2323,9 @@ int32 FProtocol4Stage::OnDataKnown(
 	FStreamReader& Reader,
 	FAnalysisBridge& Bridge)
 {
-	using namespace UE::Trace::Protocol4;
-
 	switch (Uid)
 	{
-	case EKnownEventUids::NewEvent:
+	case Protocol4::EKnownEventUids::NewEvent:
 		{
 			const auto* Size = Reader.GetPointer<uint16>();
 			const FTypeRegistry::FTypeInfo* TypeInfo = TypeRegistry.Add(Size + 1);
@@ -2338,9 +2335,9 @@ int32 FProtocol4Stage::OnDataKnown(
 		}
 		break;
 
-	case EKnownEventUids::EnterScope:
-	case EKnownEventUids::EnterScope_T:
-		if (Uid > EKnownEventUids::EnterScope)
+	case Protocol4::EKnownEventUids::EnterScope:
+	case Protocol4::EKnownEventUids::EnterScope_T:
+		if (Uid > Protocol4::EKnownEventUids::EnterScope)
 		{
 			const uint8* Stamp = Reader.GetPointer(sizeof(uint64) - 1);
 			if (Stamp == nullptr)
@@ -2359,9 +2356,9 @@ int32 FProtocol4Stage::OnDataKnown(
 		}
 		return 1;
 
-	case EKnownEventUids::LeaveScope:
-	case EKnownEventUids::LeaveScope_T:
-		if (Uid > EKnownEventUids::LeaveScope)
+	case Protocol4::EKnownEventUids::LeaveScope:
+	case Protocol4::EKnownEventUids::LeaveScope_T:
+		if (Uid > Protocol4::EKnownEventUids::LeaveScope)
 		{
 			const uint8* Stamp = Reader.GetPointer(sizeof(uint64) - 1);
 			if (Stamp == nullptr)
