Tail-tracing (a ring-buffer of previously traced events) has resulted in holes in the serial numbering used to reconstruct event order when analyzing a trace stream. Previously the serials were always consecutive, but now that tail-tracing drops older packets of events, serial number holes can appear. To mitigate that, if the same serial number is still queued in the analysis rota after two updates, a hole is assumed and the expected serial number is adjusted.

diff --git a/Engine/Source/Developer/TraceAnalysis/Private/Analysis/Engine.cpp b/Engine/Source/Developer/TraceAnalysis/Private/Analysis/Engine.cpp
--- a/Engine/Source/Developer/TraceAnalysis/Private/Analysis/Engine.cpp
+++ b/Engine/Source/Developer/TraceAnalysis/Private/Analysis/Engine.cpp
@@ -1814,6 +1814,7 @@ protected:
 	FTypeRegistry				TypeRegistry;
 	FTransport*					Transport;
 	uint32						ProtocolVersion;
+	uint32						SerialInertia = ~0u;
 };
 
 ////////////////////////////////////////////////////////////////////////////////
@@ -1992,6 +1993,28 @@ FProtocol2Stage::EStatus FProtocol2Stage::OnData(
 		}
 	}
 
+	// Patch the serial value to try and recover from gaps. Note that the bits
+	// used to wait for the serial-sync event are ignored as that event may never
+	// be reached if leading events are synchronised. One update's worth of
+	// inertia is added as the missing event can quite legitimately be in the
+	// next packet delivered.
+	if (Rota.Num() > 0)
+	{
+		int32 LowestSerial = int32(Rota[0].Serial);
+		if (LowestSerial >= 0)
+		{
+			if (SerialInertia == LowestSerial)
+			{
+				Serial.Value = LowestSerial;
+				SerialInertia = ~0u;
+			}
+			else
+			{
+				SerialInertia = LowestSerial;
+			}
+		}
+	}
+
 	return Reader.IsEmpty() ? EStatus::Eof : EStatus::NotEnoughData;
 }
 
@@ -2077,7 +2100,7 @@ int32 FProtocol2Stage::OnData(FStreamReader& Reader, FAnalysisBridge& Bridge)
 		if ((TypeInfo->Flags & FTypeRegistry::FTypeInfo::Flag_NoSync) == 0)
 		{
 			Serial.Value += 1;
-			Serial.Value &= 0x7fffffff; // don't set msb. that has other uses
+			Serial.Value &= 0x3fffffff; // Don't set MSBs. They have other uses
 		}
 
 		auto* EventData = (const uint8*)Header + BlockSize - Header->Size;
