Gaps can occur in event serial numbering when late connect happens because tail-tracing discards packets. When that happens, analysis stalls waiting for a serial number that never occurs. This case is detected by seeing what serial analysis is waiting for and if it is the same after some inertia, decide analysis was stuck. The previous inertia length was too small. Now it is bigger. Sorting of the rota also takes into account the wrap.

#rnx

diff --git a/Engine/Source/Developer/TraceAnalysis/Private/Analysis/Engine.cpp b/Engine/Source/Developer/TraceAnalysis/Private/Analysis/Engine.cpp
--- a/Engine/Source/Developer/TraceAnalysis/Private/Analysis/Engine.cpp
+++ b/Engine/Source/Developer/TraceAnalysis/Private/Analysis/Engine.cpp
@@ -1885,7 +1885,13 @@ FProtocol2Stage::EStatus FProtocol2Stage::OnData(
 		uint32				Serial;
 		uint32				ThreadId;
 		FStreamReader*		Reader;
-		bool				operator < (const FRotaItem& Rhs) { return Serial < Rhs.Serial; }
+
+		bool operator < (const FRotaItem& Rhs)
+		{
+			int32 Delta = Rhs.Serial - Serial;
+			int32 Wrapped = uint32(Delta + 0x007f'fffe) >= uint32(0x00ff'fffd);
+			return (Wrapped ^ (Delta > 0)) != 0;
+		}
 	};
 	TArray<FRotaItem> Rota;
 
@@ -1995,22 +2001,37 @@ FProtocol2Stage::EStatus FProtocol2Stage::OnData(
 
 	// Patch the serial value to try and recover from gaps. Note that the bits
 	// used to wait for the serial-sync event are ignored as that event may never
-	// be reached if leading events are synchronised. One update's worth of
-	// inertia is added as the missing event can quite legitimately be in the
-	// next packet delivered.
+	// be reached if leading events are synchronised. Some inertia is added as
+	// the missing range of events can be in subsequent packets. (Maximum inertia
+	// need only be the size of the tail / io-read-size).
 	if (Rota.Num() > 0)
 	{
 		int32 LowestSerial = int32(Rota[0].Serial);
 		if (LowestSerial >= 0)
 		{
-			if (SerialInertia == LowestSerial)
+			enum {
+				InertiaLen		= 0x20,
+				InertiaStart	= 0x7f - InertiaLen,
+				InertiaBase		= InertiaStart << 25,	// leaves bit 24 unset so
+				InertiaInc		= 1 << 25,				// ~0u can never happen
+			};
+			if (SerialInertia == ~0u)
 			{
-				Serial.Value = LowestSerial;
-				SerialInertia = ~0u;
+				SerialInertia = LowestSerial + InertiaBase;
 			}
 			else
 			{
-				SerialInertia = LowestSerial;
+				SerialInertia += InertiaInc;
+				if (int32(SerialInertia) >= 0)
+				{
+					if (SerialInertia == LowestSerial)
+					{
+						SerialInertia = ~0u;
+						Serial.Value = LowestSerial;
+						return OnData(Reader, Context);
+					}
+					SerialInertia = ~0u;
+				}
 			}
 		}
 	}
