Changed the aux-data collector to use a modestly sized inline allocator.

#rnx

diff --git a/Engine/Source/Developer/TraceAnalysis/Private/Analysis/Engine.cpp b/Engine/Source/Developer/TraceAnalysis/Private/Analysis/Engine.cpp
--- a/Engine/Source/Developer/TraceAnalysis/Private/Analysis/Engine.cpp
+++ b/Engine/Source/Developer/TraceAnalysis/Private/Analysis/Engine.cpp
@@ -56,7 +56,7 @@ struct FAuxData
 
 ////////////////////////////////////////////////////////////////////////////////
 struct FAuxDataCollector
-: public TArray<FAuxData>
+: public TArray<FAuxData, TInlineAllocator<8>>
 {
 };
 
