Added an IsUidValid() method.

#rnx

diff --git a/Engine/Source/Developer/TraceAnalysis/Private/Analysis/Engine.cpp b/Engine/Source/Developer/TraceAnalysis/Private/Analysis/Engine.cpp
--- a/Engine/Source/Developer/TraceAnalysis/Private/Analysis/Engine.cpp
+++ b/Engine/Source/Developer/TraceAnalysis/Private/Analysis/Engine.cpp
@@ -807,6 +807,7 @@ public:
 	const FTypeInfo*	Add(const void* TraceData);
 	void				Add(FTypeInfo* TypeInfo);
 	const FTypeInfo*	Get(uint32 Uid) const;
+	bool				IsUidValid(uint32 Uid) const;
 
 private:
 	TArray<FTypeInfo*>	TypeInfos;
@@ -912,6 +913,12 @@ void FTypeRegistry::Add(FTypeInfo* TypeInfo)
 	TypeInfos[Uid] = TypeInfo;
 }
 
+////////////////////////////////////////////////////////////////////////////////
+bool FTypeRegistry::IsUidValid(uint32 Uid) const
+{
+	return (Uid < uint32(TypeInfos.Num())) && (TypeInfos[Uid] != nullptr);
+}
+
 // {{{1 analyzer-hub -----------------------------------------------------------
 class FAnalyzerHub
 {
