Signed/unsigned mismatch

#rnx

diff --git a/Engine/Source/Developer/TraceAnalysis/Private/Analysis/Engine.cpp b/Engine/Source/Developer/TraceAnalysis/Private/Analysis/Engine.cpp
--- a/Engine/Source/Developer/TraceAnalysis/Private/Analysis/Engine.cpp
+++ b/Engine/Source/Developer/TraceAnalysis/Private/Analysis/Engine.cpp
@@ -897,7 +897,7 @@ void FTypeRegistry::Add(FTypeInfo* TypeInfo)
 	// Add the type to the type-infos table. Usually duplicates are an error
 	// but due to backwards compatibility we'll override existing types.
 	uint16 Uid = TypeInfo->Uid;
-	if (Uid < TypeInfos.Num())
+	if (Uid < uint32(TypeInfos.Num()))
  	{
 		if (TypeInfos[Uid] != nullptr)
  		{
