Explicitly terminate runs of aux data blocks with a well-known "aux data terminal" event. Previously a null byte was used. Each aux data block would write the null. Subsequent blocks would overwrite the prior block's null, effectively pushing the null to the end of a run aux data. However, if a different event was traced in the middle of the run it too would overwrite the null and corrupt the stream. With an explicit terminal event, structure can be recovered regardless of any multiplexing.

Serial numbers are no longer perscriptive (i.e. there can be gaps).

Normal trace events' headers are 50% smaller.

diff --git a/Engine/Source/Developer/TraceAnalysis/Private/Analysis/Engine.cpp b/Engine/Source/Developer/TraceAnalysis/Private/Analysis/Engine.cpp
--- a/Engine/Source/Developer/TraceAnalysis/Private/Analysis/Engine.cpp
+++ b/Engine/Source/Developer/TraceAnalysis/Private/Analysis/Engine.cpp
@@ -6,6 +6,7 @@
 #include "Containers/ArrayView.h"
 #include "CoreGlobals.h"
 #include "HAL/UnrealMemory.h"
+#include "Logging/LogMacros.h"
 #include "StreamReader.h"
 #include "Templates/UnrealTemplate.h"
 #include "Trace/Analysis.h"
@@ -2434,6 +2435,648 @@ int32 FProtocol4Stage::OnDataKnown(
 
 
 
+// {{{1 protocol-5 -------------------------------------------------------------
+
+////////////////////////////////////////////////////////////////////////////////
+DEFINE_LOG_CATEGORY_STATIC(LogTraceAnalysis, Log, All)
+
+#if 0
+#	define TRACE_ANALYSIS_DEBUG(Format, ...) \
+		do { UE_LOG(LogTraceAnalysis, Log, TEXT(Format), __VA_ARGS__) } while (0)
+#else
+#	define TRACE_ANALYSIS_DEBUG(...)
+#endif
+
+////////////////////////////////////////////////////////////////////////////////
+class FProtocol5Stage
+	: public FAnalysisMachine::FStage
+{
+public:
+							FProtocol5Stage(FTransport* InTransport);
+
+private:
+	struct alignas(16) FEventDesc
+	{
+		int32				Serial;
+		uint16				Uid			: 15;
+		uint16				bHasAux		: 1;
+		uint16				AuxKey;
+		const uint8*		Data;
+	};
+	static_assert(sizeof(FEventDesc) == 16, "");
+
+	enum ESerial : int32 
+	{
+		Bits		= 24,
+		Mask		= (1 << Bits) - 1,
+		Range		= 1 << Bits,
+		HalfRange	= Range >> 1,
+		Ignored		= Range << 2, // far away so proper serials always compare less-than
+		Terminal,
+	};
+
+	using EventDescArray	= TArray<FEventDesc>;
+	using EKnownUids		= Protocol5::EKnownEventUids;
+
+
+	virtual EStatus			OnData(FStreamReader& Reader, const FMachineContext& Context) override;
+	EStatus					OnDataImportant(const FMachineContext& Context);
+	EStatus					OnDataNormal(const FMachineContext& Context);
+	int32					ParseImportantEvents(FStreamReader& Reader, EventDescArray& OutEventDescs);
+	int32					ParseEvents(FStreamReader& Reader, EventDescArray& OutEventDescs);
+	int32					ParseEventsWithAux(FStreamReader& Reader, EventDescArray& OutEventDescs);
+	int32					ParseEvent(FStreamReader& Reader, FEventDesc& OutEventDesc);
+	int						DispatchEvents(FAnalysisBridge& Bridge, const FEventDesc* EventDesc, uint32 Count);
+	FTypeRegistry			TypeRegistry;
+	FTidPacketTransport&	Transport;
+	EventDescArray			EventDescs;
+};
+
+////////////////////////////////////////////////////////////////////////////////
+FProtocol5Stage::FProtocol5Stage(FTransport* InTransport)
+: Transport(*(FTidPacketTransport*)InTransport)
+{
+	EventDescs.Reserve(8 << 10);
+}
+
+////////////////////////////////////////////////////////////////////////////////
+FProtocol5Stage::EStatus FProtocol5Stage::OnData(
+	FStreamReader& Reader,
+	const FMachineContext& Context)
+{
+	Transport.SetReader(Reader);
+	Transport.Update();
+
+	EStatus Ret = OnDataImportant(Context);
+	if (Ret != EStatus::Eof)
+	{
+		return Ret;
+	}
+
+	return OnDataNormal(Context);
+}
+
+////////////////////////////////////////////////////////////////////////////////
+FProtocol5Stage::EStatus FProtocol5Stage::OnDataImportant(const FMachineContext& Context)
+{
+	// New and important event packets
+
+	EventDescs.Reset();
+	bool bNotEnoughData = false;
+
+	for (uint32 i = 0, n = ETransportTid::Bias; i < n; ++i)
+	{
+		FStreamReader* ThreadReader = Transport.GetThreadStream(i);
+		if (ThreadReader->IsEmpty())
+		{
+			continue;
+		}
+
+		if (ParseImportantEvents(*ThreadReader, EventDescs) < 0)
+		{
+			return EStatus::Error;
+		}
+
+		bNotEnoughData |= !ThreadReader->IsEmpty();
+	}
+
+	// Dispatch and filter out new-event events 
+	for (uint32 j = 0, k = 0, m = EventDescs.Num(); ; ++j)
+	{
+		if (j >= m)
+		{
+			EventDescs.SetNum(k);
+			break;
+		}
+
+		const FEventDesc& EventDesc = EventDescs[j];
+		if (EventDesc.Uid != EKnownUids::NewEvent)
+		{
+			EventDescs[k] = EventDesc;
+			++k;
+			continue;
+		}
+
+		const FTypeRegistry::FTypeInfo* TypeInfo = TypeRegistry.Add(EventDesc.Data);
+		Context.Bridge.OnNewType(TypeInfo);
+	}
+
+	if (EventDescs.Num() <= 0)
+	{
+		return bNotEnoughData ? EStatus::NotEnoughData : EStatus::Eof;
+	}
+
+	// Dispatch looks ahead to the next desc looking for runs of aux blobs. As
+	// such we should add a terminal desc for it to read. Note the "- 1" too.
+	FEventDesc& EventDesc = EventDescs.Emplace_GetRef();
+	EventDesc.Serial = ESerial::Terminal;
+
+	Context.Bridge.SetActiveThread(ETransportTid::Importants);
+	if (DispatchEvents(Context.Bridge, EventDescs.GetData(), EventDescs.Num() - 1) < 0)
+	{
+		return EStatus::Error;
+	}
+
+	return bNotEnoughData ? EStatus::NotEnoughData : EStatus::Eof;
+}
+
+////////////////////////////////////////////////////////////////////////////////
+int32 FProtocol5Stage::ParseImportantEvents(FStreamReader& Reader, EventDescArray& OutEventDescs)
+{
+	using namespace Protocol5;
+
+	while (true)
+	{
+		int32 Remaining = Reader.GetRemaining();
+		if (Remaining < sizeof(FImportantEventHeader))
+		{
+			return 1;
+		}
+
+		const auto* Header = Reader.GetPointerUnchecked<FImportantEventHeader>();
+		if (Remaining < int32(Header->Size) + sizeof(FImportantEventHeader))
+		{
+			return 1;
+		}
+
+		uint32 Uid = Header->Uid;
+
+		FEventDesc EventDesc;
+		EventDesc.Uid = Uid;
+		EventDesc.Data = Header->Data;
+		EventDesc.Serial = ESerial::Ignored;
+
+		// Special case for new events. It would work to add a 0 type to the
+		// registry but this way avoid raveling things together.
+		if (Uid == EKnownUids::NewEvent)
+		{
+			OutEventDescs.Add(EventDesc);
+			Reader.Advance(sizeof(*Header) + Header->Size);
+			continue;
+		}
+
+		const FTypeRegistry::FTypeInfo* TypeInfo = TypeRegistry.Get(Uid);
+		if (TypeInfo == nullptr)
+		{
+			UE_LOG(LogTraceAnalysis, Log, TEXT("Unknown event UID: %08x"), Uid);
+			return 1;
+		}
+
+		OutEventDescs.Add(EventDesc);
+
+		if (TypeInfo->Flags & FTypeRegistry::FTypeInfo::Flag_MaybeHasAux)
+		{
+			EventDesc.bHasAux = 1;
+
+			const uint8* Cursor = Header->Data + TypeInfo->EventSize;
+			const uint8* End = Header->Data + Header->Size;
+			while (Cursor <= End)
+			{
+				if (Cursor[0] == uint8(EKnownUids::AuxDataTerminal))
+				{
+					break;
+				}
+
+				const auto* AuxHeader = (FAuxHeader*)Cursor;
+
+				FEventDesc& AuxDesc = OutEventDescs.Emplace_GetRef();
+				AuxDesc.Uid = uint8(EKnownUids::AuxData);
+				AuxDesc.Data = AuxHeader->Data;
+
+				Cursor = AuxHeader->Data + (AuxHeader->Pack >> FAuxHeader::SizeShift);
+			}
+
+			if (Cursor[0] != uint8(EKnownUids::AuxDataTerminal))
+			{
+				UE_LOG(LogTraceAnalysis, Warning, TEXT("Expecting AuxDataTerminal event"));
+				return -1;
+			}
+		}
+
+		Reader.Advance(sizeof(*Header) + Header->Size);
+	}
+
+	return 0;
+}
+
+////////////////////////////////////////////////////////////////////////////////
+FProtocol5Stage::EStatus FProtocol5Stage::OnDataNormal(const FMachineContext& Context)
+{
+	// Ordinary events
+
+	EventDescs.Reset();
+	bool bNotEnoughData = false;
+
+	struct FEventDescStream
+	{
+		uint32				ThreadId;
+		uint32				Index;
+		const FEventDesc*	EventDescs;
+	};
+	TArray<FEventDescStream> EventDescHeap;
+	EventDescHeap.Reserve(Transport.GetThreadCount());
+
+	for (uint32 i = ETransportTid::Bias, n = Transport.GetThreadCount(); i < n; ++i)
+	{
+		uint32 NumEventDescs = EventDescs.Num();
+
+		TRACE_ANALYSIS_DEBUG("Thread: %03d Id:%04x", i, Transport.GetThreadId(i));
+		
+		// Extract all the events in the stream for this thread
+		FStreamReader* ThreadReader = Transport.GetThreadStream(i);
+		if (ParseEvents(*ThreadReader, EventDescs) < 0)
+		{
+			return EStatus::Error;
+		}
+
+		bNotEnoughData |= !ThreadReader->IsEmpty();
+
+		if (uint32(EventDescs.Num()) != NumEventDescs)
+		{
+			// Add a dummy event to delineate the end of this thread's events
+			FEventDesc& EventDesc = EventDescs.Emplace_GetRef();
+			EventDesc.Serial = ESerial::Terminal;
+
+			EventDescHeap.Add({Transport.GetThreadId(i), NumEventDescs});
+		}
+
+		TRACE_ANALYSIS_DEBUG("Thread: %03d bNotEnoughData:%d", i, bNotEnoughData);
+	}
+
+	// Early out if there isn't any events available.
+	if (UNLIKELY(EventDescHeap.IsEmpty()))
+	{
+		return bNotEnoughData ? EStatus::NotEnoughData : EStatus::Eof;
+	}
+
+	// Now EventDescs is stable we can convert the indices into pointers
+	for (FEventDescStream& Stream : EventDescHeap)
+	{
+		Stream.EventDescs = EventDescs.GetData() + Stream.Index;
+	}
+
+	// Provided that less than approximately "SerialRange * BytesPerSerial"
+	// is buffered there should never be more that "SerialRange / 2" serial
+	// numbers. Thus if the distance between any two serial numbers is larger
+	// than half the serial space, they have wrapped.
+
+	// A min-heap is used to peel off groups of events by lowest serial
+	auto Comparison = [] (const FEventDescStream& Lhs, const FEventDescStream& Rhs)
+	{
+		int32 Delta = Rhs.EventDescs->Serial - Lhs.EventDescs->Serial;
+		int32 Wrapped = uint32(Delta + ESerial::HalfRange - 1) >= uint32(ESerial::Range - 2);
+		return (Wrapped ^ (Delta > 0)) != 0;
+	};
+	EventDescHeap.Heapify(Comparison);
+
+	do
+	{
+		const FEventDescStream& Stream = EventDescHeap.HeapTop();
+		const FEventDesc* StartDesc = Stream.EventDescs;
+		const FEventDesc* EndDesc = StartDesc;
+
+		Context.Bridge.SetActiveThread(Stream.ThreadId);
+
+		// Extract a run of consecutive events (plus runs of unsynchronised ones)
+		for (; EndDesc->Serial == ESerial::Ignored; ++EndDesc);
+		if (EndDesc->Serial != ESerial::Terminal)
+		{
+			uint32 NextSerial = EndDesc->Serial;
+			do
+			{
+				do
+				{
+					++EndDesc;
+				}
+				while (EndDesc->Serial == ESerial::Ignored);
+
+				++NextSerial;
+			}
+			while (EndDesc->Serial == NextSerial);
+		}
+
+		// Update the heap
+		if (EndDesc->Serial != ESerial::Terminal)
+		{
+			EventDescHeap.Add({Stream.ThreadId, 0, EndDesc});
+		}
+		EventDescHeap.HeapPopDiscard(Comparison);
+
+		// Dispatch.
+		int32 DescNum = int32(UPTRINT(EndDesc - StartDesc));
+		check(DescNum > 0);
+		if (DispatchEvents(Context.Bridge, StartDesc, DescNum) < 0)
+		{
+			return EStatus::Error;
+		}
+	}
+	while (!EventDescHeap.IsEmpty());
+
+	return bNotEnoughData ? EStatus::NotEnoughData : EStatus::Eof;
+}
+
+////////////////////////////////////////////////////////////////////////////////
+int32 FProtocol5Stage::ParseEvents(FStreamReader& Reader, EventDescArray& OutEventDescs)
+{
+	while (!Reader.IsEmpty())
+	{
+		FEventDesc EventDesc;
+		int32 Size = ParseEvent(Reader, EventDesc);
+		if (Size <= 0)
+		{
+			return Size;
+		}
+
+		TRACE_ANALYSIS_DEBUG("Event: %04d Uid:%04x Serial:%08x", OutEventDescs.Num(), EventDesc.Uid, EventDesc.Serial);
+		OutEventDescs.Add(EventDesc);
+
+		if (EventDesc.bHasAux)
+		{
+			uint32 RewindDescsNum = OutEventDescs.Num() - 1;
+			auto RewindMark = Reader.SaveMark();
+
+			Reader.Advance(Size);
+
+			int Ok = ParseEventsWithAux(Reader, OutEventDescs);
+			if (Ok < 0)
+			{
+				return Ok;
+			}
+
+			if (Ok == 0)
+			{
+				OutEventDescs.SetNum(RewindDescsNum);
+				Reader.RestoreMark(RewindMark);
+				break;
+			}
+
+			continue;
+		}
+
+		Reader.Advance(Size);
+	}
+
+	return 0;
+}
+
+////////////////////////////////////////////////////////////////////////////////
+int32 FProtocol5Stage::ParseEventsWithAux(FStreamReader& Reader, EventDescArray& OutEventDescs)
+{
+	// We are now "in" the scope of an event with zero or more aux-data blocks.
+	// We will consume events until we leave this scope (a aux-data-terminal).
+	// A running key is assigned to each event with a gap left following events
+	// that may have aux-data blocks. Aux-data blocks are assigned a key that
+	// fits in these gaps. Once sorted by this key, events maintain their order
+	// while aux-data blocks are moved to directly follow their owners.
+
+	TArray<uint16, TInlineAllocator<8>> AuxKeyStack = { 0 };
+	uint32 AuxKey = 2;
+
+	uint32 FirstDescIndex = OutEventDescs.Num();
+	bool bUnsorted = false;
+
+	while (!Reader.IsEmpty())
+	{
+		FEventDesc EventDesc;
+
+		int32 Size = ParseEvent(Reader, EventDesc);
+		if (Size <= 0)
+		{
+			return Size;
+		}
+
+		TRACE_ANALYSIS_DEBUG("Event: %04d Uid:%04x Serial:%08x", OutEventDescs.Num(), EventDesc.Uid, EventDesc.Serial);
+
+		Reader.Advance(Size);
+
+		if (EventDesc.Uid == EKnownUids::AuxDataTerminal)
+		{
+			// Leave the scope of an aux-owning event.
+			if (AuxKeyStack.Pop() == 0)
+			{
+				break;
+			}
+			continue;
+		}
+		else if (EventDesc.Uid == EKnownUids::AuxData)
+		{
+			// Move an aux-data block to follow its owning event
+			EventDesc.AuxKey = AuxKeyStack.Last() + 1;
+		}
+		else
+		{
+			EventDesc.AuxKey = uint16(AuxKey);
+
+			// Maybe it is time to create a new aux-data owner scope
+			if (EventDesc.bHasAux)
+			{
+				AuxKeyStack.Add(uint16(AuxKey));
+				bUnsorted = true;
+
+				++AuxKey; // Add a gap for aux-blocks to sort into
+			}
+		}
+
+		OutEventDescs.Add(EventDesc);
+
+		++AuxKey;
+	}
+
+	if (AuxKeyStack.Num() > 0)
+	{
+		// There was not enough data available to complete the outer most scope
+		return 0;
+	}
+
+	checkf((AuxKey & 0xffff0000) == 0, TEXT("AuxKey overflow (%08x)"), AuxKey);
+
+	// Sort to get all aux-blocks contiguous with their owning event
+	if (bUnsorted)
+	{
+		uint32 NumDescs = OutEventDescs.Num() - FirstDescIndex;
+		TArrayView<FEventDesc> DescsView(OutEventDescs.GetData() + FirstDescIndex, NumDescs);
+		Algo::Sort(
+			DescsView,
+			[] (const FEventDesc& Lhs, const FEventDesc& Rhs)
+			{
+				return Lhs.AuxKey < Rhs.AuxKey;
+			}
+		);
+	}
+
+	return 1;
+}
+
+////////////////////////////////////////////////////////////////////////////////
+int32 FProtocol5Stage::ParseEvent(FStreamReader& Reader, FEventDesc& EventDesc)
+{
+	using namespace Protocol5;
+
+	// No need to aggressively bounds check here. Events are never fragmented
+	// due to the way that data is transported (aux payloads can be though).
+	const uint8* Cursor = Reader.GetPointerUnchecked<uint8>();
+
+	// Parse the event's ID
+	uint32 Uid = *Cursor;
+	if (Uid & EKnownUids::Flag_TwoByteUid)
+	{
+		Uid = *(uint16*)Cursor;
+		++Cursor;
+	}
+	Uid >>= EKnownUids::_UidShift;
+	++Cursor;
+
+	// Calculate the size of the event
+	uint32 Serial = uint32(ESerial::Ignored);
+	uint32 EventSize = 0;
+	if (Uid < EKnownUids::User)
+	{
+		/* Well-known event */
+
+		if (Uid == Protocol5::EKnownEventUids::AuxData)
+		{
+			--Cursor; // FAuxHeader includes the one-byte Uid
+			const auto* AuxHeader = (FAuxHeader*)Cursor;
+
+			uint32 Remaining = Reader.GetRemaining();
+			uint32 Size = AuxHeader->Pack >> FAuxHeader::SizeShift;
+			if (Remaining < Size + sizeof(FAuxHeader))
+			{
+				return 0;
+			}
+
+			EventSize = Size;
+			Cursor += sizeof(FAuxHeader);
+		}
+		else
+		{
+			switch (Uid)
+			{
+			case EKnownUids::EnterScope_T:
+			case EKnownUids::LeaveScope_T:
+				EventSize = 7;
+			};
+		}
+
+		EventDesc.bHasAux = 0;
+	}
+	else
+	{
+		/* Ordinary events */
+
+		const FTypeRegistry::FTypeInfo* TypeInfo = TypeRegistry.Get(Uid);
+		if (TypeInfo == nullptr)
+		{
+			UE_LOG(LogTraceAnalysis, Log, TEXT("Unknown event UID: %08x"), Uid);
+			return 0;
+		}
+
+		EventSize = TypeInfo->EventSize;
+		EventDesc.bHasAux = !!(TypeInfo->Flags & FTypeRegistry::FTypeInfo::Flag_MaybeHasAux);
+
+		if ((TypeInfo->Flags & FDispatch::Flag_NoSync) == 0)
+		{
+			memcpy(&Serial, Cursor, sizeof(int32));
+			Serial &= ESerial::Mask;
+			Cursor += 3;
+		}
+	}
+
+	EventDesc.Serial = Serial;
+	EventDesc.Uid = Uid;
+	EventDesc.Data = Cursor;
+
+	uint32 HeaderSize = uint32(UPTRINT(Cursor - Reader.GetPointer<uint8>()));
+	return HeaderSize + EventSize;
+}
+
+////////////////////////////////////////////////////////////////////////////////
+int FProtocol5Stage::DispatchEvents(
+	FAnalysisBridge& Bridge,
+	const FEventDesc* EventDesc, uint32 Count)
+{
+	using namespace Protocol5;
+
+	FAuxDataCollector AuxCollector;
+
+	for (const FEventDesc *Cursor = EventDesc, *End = EventDesc + Count; Cursor < End; ++Cursor)
+	{
+		// Maybe this is a "well-known" event that is handled a little different?
+		switch (Cursor->Uid)
+		{
+		case EKnownUids::EnterScope_T:
+		{
+			uint64 Timestamp = *(uint64*)(Cursor->Data - 1) >> 8;
+			Bridge.EnterScope(Timestamp);
+			continue;
+		}
+
+		case EKnownUids::LeaveScope_T:
+		{
+			uint64 Timestamp = *(uint64*)(Cursor->Data - 1) >> 8;
+			Bridge.LeaveScope(Timestamp);
+			continue;
+		}
+
+		case EKnownUids::EnterScope: Bridge.EnterScope(); continue;
+		case EKnownUids::LeaveScope: Bridge.LeaveScope(); continue;
+
+		case EKnownUids::AuxData:
+		case EKnownUids::AuxDataTerminal:
+			continue;
+
+		default:
+			if (!TypeRegistry.IsUidValid(Cursor->Uid))
+			{
+				UE_LOG(LogTraceAnalysis, Warning, TEXT("Unexpected event UID: %08x"), Cursor->Uid);
+				return -1;
+			}
+			break;
+		}
+
+		// It is a normal event.
+		const FTypeRegistry::FTypeInfo* TypeInfo = TypeRegistry.Get(Cursor->Uid);
+		FEventDataInfo EventDataInfo = {
+			Cursor->Data,
+			*TypeInfo,
+			&AuxCollector,
+			TypeInfo->EventSize,
+		};
+
+		// Gather its auxiliary data blocks into a collector.
+		if (Cursor->bHasAux)
+		{
+			while (true)
+			{
+				++Cursor;
+
+				if (Cursor->Uid != EKnownUids::AuxData)
+				{
+					--Cursor; // Read off too much. Put it back.
+					break;
+				}
+
+				const auto* AuxHeader = ((FAuxHeader*)(Cursor->Data)) - 1;
+
+				FAuxData AuxData;
+				AuxData.Data = AuxHeader->Data;
+				AuxData.DataSize = (AuxHeader->Pack >> FAuxHeader::SizeShift);
+				AuxData.FieldIndex = AuxHeader->FieldIndex_Size & FAuxHeader::FieldMask;
+				// AuxData.FieldSizeAndType = ... - this is assigned on demand in GetData()
+				AuxCollector.Add(AuxData);
+			}
+		}
+
+		Bridge.OnEvent(EventDataInfo);
+
+		AuxCollector.Reset();
+	}
+
+	return 0;
+}
+
+
+
 // {{{1 est.-transport ---------------------------------------------------------
 
 ////////////////////////////////////////////////////////////////////////////////
@@ -2489,6 +3132,11 @@ FEstablishTransportStage::EStatus FEstablishTransportStage::OnData(
 		Context.Machine.Transition();
 		break;
 
+	case Protocol5::EProtocol::Id:
+		Context.Machine.QueueStage<FProtocol5Stage>(Transport);
+		Context.Machine.Transition();
+		break;
+
 	default:
 		return EStatus::Error;
 	}
diff --git a/Engine/Source/Runtime/TraceLog/Private/Trace/EventNode.cpp b/Engine/Source/Runtime/TraceLog/Private/Trace/EventNode.cpp
--- a/Engine/Source/Runtime/TraceLog/Private/Trace/EventNode.cpp
+++ b/Engine/Source/Runtime/TraceLog/Private/Trace/EventNode.cpp
@@ -117,20 +117,17 @@ void FEventNode::Describe() const
 	}
 
 	// Allocate the new event event in the log stream.
-	uint16 EventSize = sizeof(FNewEventEvent);
+	uint32 EventSize = sizeof(FNewEventEvent);
 	EventSize += sizeof(FNewEventEvent::Fields[0]) * Info->FieldCount;
 	EventSize += NamesSize;
 
-	struct FNewEventDummy
-	{
-		static constexpr uint32 GetSize()	{ return 0; }
-		static constexpr uint32 GetUid()	{ return 0; }
-		enum { EventFlags = FEventInfo::Flag_NoSync };
-	};
-	FLogScope LogScope = FLogScope::Enter<FNewEventDummy>(EventSize);
-	auto& Event = *(FNewEventEvent*)(LogScope.GetPointer());
+	FLogScope LogScope = FLogScope::EnterImpl<FEventInfo::Flag_NoSync>(0, EventSize + sizeof(uint16));
+	auto* Ptr = (uint16*)(LogScope.GetPointer());
+	Ptr[-1] = EKnownEventUids::NewEvent; // Make event look like an important one. Ideally they are sent
+	Ptr[ 0] = uint16(EventSize);		 // as important and not Writer_DescribeEvents()'s redirected buf.
 
 	// Write event's main properties.
+	auto& Event = *(FNewEventEvent*)(Ptr + 1);
 	Event.EventUid = uint16(Uid) >> EKnownEventUids::_UidShift;
 	Event.LoggerNameSize = LoggerName.Length;
 	Event.EventNameSize = EventName.Length;
diff --git a/Engine/Source/Runtime/TraceLog/Private/Trace/Field.cpp b/Engine/Source/Runtime/TraceLog/Private/Trace/Field.cpp
--- a/Engine/Source/Runtime/TraceLog/Private/Trace/Field.cpp
+++ b/Engine/Source/Runtime/TraceLog/Private/Trace/Field.cpp
@@ -14,20 +14,27 @@ namespace Private {
 template <typename CallbackType>
 static void Field_WriteAuxData(uint32 Index, int32 Size, CallbackType&& Callback)
 {
-	static_assert(sizeof(Private::FWriteBuffer::Overflow) >= sizeof(FAuxHeader), "FWriteBuffer::Overflow is not large enough");
+	static_assert(
+		sizeof(Private::FWriteBuffer::Overflow) >= sizeof(FAuxHeader) + sizeof(uint8 /*AuxDataTerminal*/),
+		"FWriteBuffer::Overflow is not large enough"
+	);
+
+	// Early-out if there would be nothing to write
+	if (Size == 0)
+	{
+		return;
+	}
 
 	// Header
-	const int bMaybeHasAux = true;
 	FWriteBuffer* Buffer = Writer_GetBuffer();
-	Buffer->Cursor += sizeof(FAuxHeader) - bMaybeHasAux;
-
-	auto* Header = (FAuxHeader*)(Buffer->Cursor - sizeof(FAuxHeader));
-	Header->Size = Size << 8;
-	Header->FieldIndex = uint8(0x80 | (Index & int(EIndexPack::NumFieldsMask)));
-
-	bool bCommit = ((uint8*)Header + bMaybeHasAux == Buffer->Committed);
+	auto* Header = (FAuxHeader*)(Buffer->Cursor);
+	Header->Pack = Size << FAuxHeader::SizeShift;
+	Header->Pack |= Index << FAuxHeader::FieldShift;
+	Header->Uid = uint8(EKnownEventUids::AuxData) << EKnownEventUids::_UidShift;
+	Buffer->Cursor += sizeof(FAuxHeader);
 
 	// Array data
+	bool bCommit = ((uint8*)Header == Buffer->Committed);
 	while (true)
 	{
 		if (Buffer->Cursor >= (uint8*)Buffer)
@@ -55,10 +62,6 @@ static void Field_WriteAuxData(uint32 Index, int32 Size, CallbackType&& Callback
 		}
 	}
 
-	// The auxilary data null terminator.
-	Buffer->Cursor[0] = 0;
-	Buffer->Cursor++;
-
 	if (bCommit)
 	{
 		AtomicStoreRelease(&(uint8* volatile&)(Buffer->Committed), Buffer->Cursor);
diff --git a/Engine/Source/Runtime/TraceLog/Private/Trace/TlsBuffer.cpp b/Engine/Source/Runtime/TraceLog/Private/Trace/TlsBuffer.cpp
--- a/Engine/Source/Runtime/TraceLog/Private/Trace/TlsBuffer.cpp
+++ b/Engine/Source/Runtime/TraceLog/Private/Trace/TlsBuffer.cpp
@@ -44,7 +44,7 @@ UE_TRACE_EVENT_END()
 
 ////////////////////////////////////////////////////////////////////////////////
 #define T_ALIGN alignas(PLATFORM_CACHE_LINE_SIZE)
-static FWriteBuffer						GNullWriteBuffer	= { 0, 0, 0, 0, nullptr, nullptr, (uint8*)&GNullWriteBuffer };
+static FWriteBuffer						GNullWriteBuffer	= { {}, 0, 0, nullptr, nullptr, (uint8*)&GNullWriteBuffer };
 thread_local FWriteBuffer*				GTlsWriteBuffer		= &GNullWriteBuffer;
 static FWriteBuffer* __restrict			GActiveThreadList;	// = nullptr;
 T_ALIGN static FWriteBuffer* volatile	GNewThreadList;		// = nullptr;
diff --git a/Engine/Source/Runtime/TraceLog/Private/Trace/Writer.cpp b/Engine/Source/Runtime/TraceLog/Private/Trace/Writer.cpp
--- a/Engine/Source/Runtime/TraceLog/Private/Trace/Writer.cpp
+++ b/Engine/Source/Runtime/TraceLog/Private/Trace/Writer.cpp
@@ -53,17 +53,11 @@ void			Writer_ShutdownControl();
 
 
 ////////////////////////////////////////////////////////////////////////////////
-UE_TRACE_EVENT_BEGIN($Trace, NewTrace, NoSync)
+UE_TRACE_EVENT_BEGIN($Trace, NewTrace, Important|NoSync)
 	UE_TRACE_EVENT_FIELD(uint64, StartCycle)
 	UE_TRACE_EVENT_FIELD(uint64, CycleFrequency)
-	UE_TRACE_EVENT_FIELD(uint32, Serial)
-	UE_TRACE_EVENT_FIELD(uint16, UserUidBias)
 	UE_TRACE_EVENT_FIELD(uint16, Endian)
 	UE_TRACE_EVENT_FIELD(uint8, PointerSize)
-	UE_TRACE_EVENT_FIELD(uint8, FeatureSet)
-UE_TRACE_EVENT_END()
-
-UE_TRACE_EVENT_BEGIN($Trace, SerialSync, NoSync)
 UE_TRACE_EVENT_END()
 
 
@@ -233,7 +227,6 @@ void Writer_MemoryFree(void* Address, uint32 Size)
 
 ////////////////////////////////////////////////////////////////////////////////
 static UPTRINT					GDataHandle;		// = 0
-static bool						GSerialSyncPending;	// = false
 UPTRINT							GPendingDataHandle;	// = 0
 
 ////////////////////////////////////////////////////////////////////////////////
@@ -354,14 +347,6 @@ static void Writer_DescribeAnnounce()
 ////////////////////////////////////////////////////////////////////////////////
 static void Writer_LogHeader()
 {
-	UE_TRACE_LOG($Trace, NewTrace, TraceLogChannel)
-		<< NewTrace.Serial(AtomicLoadRelaxed(&GLogSerial))
-		<< NewTrace.UserUidBias(EKnownEventUids::User)
-		<< NewTrace.Endian(0x524d)
-		<< NewTrace.PointerSize(sizeof(void*))
-		<< NewTrace.StartCycle(GStartCycle)
-		<< NewTrace.CycleFrequency(TimeGetFrequency())
-		<< NewTrace.FeatureSet(1);
 }
 
 ////////////////////////////////////////////////////////////////////////////////
@@ -398,7 +383,6 @@ static bool Writer_UpdateConnection()
 	{
 		IoClose(GPendingDataHandle);
 		GPendingDataHandle = 0;
-		GSerialSyncPending = false;
 		return false;
 	}
 
@@ -443,16 +427,10 @@ static bool Writer_UpdateConnection()
 	FEventNode::OnConnect();
 	Writer_DescribeEvents();
 
-	// Send the header event
-	TWriteBufferRedirect<512> HeaderEvents;
-	Writer_LogHeader();
-	HeaderEvents.Close();
-	Writer_SendData(ETransportTid::Internal, HeaderEvents.GetData(), HeaderEvents.GetSize());
-
+	// Send cached events (i.e. importants) and the tail of recent events
 	Writer_CacheOnConnect();
 	Writer_TailOnConnect();
 
-	GSerialSyncPending = true;
 	return true;
 }
 
@@ -470,17 +448,6 @@ static void Writer_WorkerUpdate()
 	Writer_DescribeAnnounce();
 	Writer_UpdateSharedBuffers();
 	Writer_DrainBuffers();
-
-	if (GSerialSyncPending)
-	{
-		GSerialSyncPending = false;
-
-		// When analysis receives the "serial sync" event the starting serial
-		// can be established from preceding synced events.
-		TWriteBufferRedirect<32> SideBuffer;
-		UE_TRACE_LOG($Trace, SerialSync, TraceLogChannel);
-		Writer_SendData(ETransportTid::Internal, SideBuffer.GetData(), SideBuffer.GetSize());
-	}
 }
 
 ////////////////////////////////////////////////////////////////////////////////
@@ -542,15 +509,11 @@ static void Writer_InternalInitializeImpl()
 	Writer_InitializePool();
 	Writer_InitializeControl();
 
-	// The order of the events at the beginning of a trace is a little sensitive,
-	// especially for the two following events that are added to the trace stream
-	// at very specific points. By dummy-tracing them we invoke the code to IDs
-	// and can guarantee the events are described on connection. It's not pretty.
-	{
-		TWriteBufferRedirect<64> TempBuffer;
-		UE_TRACE_LOG($Trace, NewTrace, TraceLogChannel);
-		UE_TRACE_LOG($Trace, SerialSync, TraceLogChannel);
-	}
+	UE_TRACE_LOG($Trace, NewTrace, TraceLogChannel)
+		<< NewTrace.StartCycle(GStartCycle)
+		<< NewTrace.CycleFrequency(TimeGetFrequency())
+		<< NewTrace.Endian(0x524d)
+		<< NewTrace.PointerSize(sizeof(void*));
 }
 
 ////////////////////////////////////////////////////////////////////////////////
diff --git a/Engine/Source/Runtime/TraceLog/Public/Trace/Detail/Field.h b/Engine/Source/Runtime/TraceLog/Public/Trace/Detail/Field.h
--- a/Engine/Source/Runtime/TraceLog/Public/Trace/Detail/Field.h
+++ b/Engine/Source/Runtime/TraceLog/Public/Trace/Detail/Field.h
@@ -114,6 +114,7 @@ template <int InIndex, int InOffset, typename Type> struct TField;
 
 enum class EIndexPack
 {
+	NumFieldsMax	= 1 << FAuxHeader::FieldBits,
 	NumFieldsShift	= 8,
 	NumFieldsMask	= (1 << NumFieldsShift) - 1,
 	AuxFieldCounter	= 1 << NumFieldsShift,
@@ -128,7 +129,7 @@ enum class EIndexPack
 			Tid		= TFieldType<Type>::Tid, \
 			Size	= TFieldType<Type>::Size, \
 		}; \
-		static_assert((Index & int(EIndexPack::NumFieldsMask)) <= 127, "Trace events may only have up to a maximum of 127 fields"); \
+		static_assert((Index & int(EIndexPack::NumFieldsMask)) < int(EIndexPack::NumFieldsMax), "Trace events may only have up to EIndexPack::NumFieldsMax fields"); \
 	private: \
 		FFieldDesc FieldDesc; \
 	public: \
diff --git a/Engine/Source/Runtime/TraceLog/Public/Trace/Detail/Important/ImportantLogScope.inl b/Engine/Source/Runtime/TraceLog/Public/Trace/Detail/Important/ImportantLogScope.inl
--- a/Engine/Source/Runtime/TraceLog/Public/Trace/Detail/Important/ImportantLogScope.inl
+++ b/Engine/Source/Runtime/TraceLog/Public/Trace/Detail/Important/ImportantLogScope.inl
@@ -30,14 +30,14 @@ FORCENOINLINE FImportantLogScope FImportantLogScope::Enter(uint32 ArrayDataSize)
 	static_assert(T::EventFlags & FEventInfo::Flag_MaybeHasAux, "Only important trace events with array-type fields need a size parameter to UE_TRACE_LOG()");
 
 	ArrayDataSize += sizeof(FAuxHeader) * T::EventProps_Meta::NumAuxFields;
-	ArrayDataSize += 1; // null terminator
+	ArrayDataSize += 1; // for AuxDataTerminal
 
 	uint32 Size = T::GetSize();
-	uint32 Uid = T::GetUid();
+	uint32 Uid = T::GetUid() >> EKnownEventUids::_UidShift;
 	FImportantLogScope Ret = EnterImpl(Uid, Size + ArrayDataSize);
 
 	Ret.AuxCursor += Size;
-	Ret.Ptr[Ret.AuxCursor] = 0; // null terminator
+	Ret.Ptr[Ret.AuxCursor] = uint8(EKnownEventUids::AuxDataTerminal);
 	return Ret;
 }
 
@@ -48,7 +48,7 @@ inline FImportantLogScope FImportantLogScope::Enter()
 	static_assert(!(T::EventFlags & FEventInfo::Flag_MaybeHasAux), "Important trace events with array-type fields must be traced with UE_TRACE_LOG(Logger, Event, Channel, ArrayDataSize)");
 
 	uint32 Size = T::GetSize();
-	uint32 Uid = T::GetUid();
+	uint32 Uid = T::GetUid() >> EKnownEventUids::_UidShift;
 	return EnterImpl(Uid, Size);
 }
 
@@ -58,7 +58,7 @@ inline FImportantLogScope FImportantLogScope::EnterImpl(uint32 Uid, uint32 Size)
 	FSharedBuffer* Buffer = AtomicLoadAcquire(&GSharedBuffer);
 
 	int32 AllocSize = Size;
-	AllocSize += sizeof(FEventHeader);
+	AllocSize += sizeof(FImportantEventHeader);
 
 	// Claim some space in the buffer
 	int32 NegSizeAndRef = 0 - ((AllocSize << FSharedBuffer::CursorShift) | FSharedBuffer::RefBit);
@@ -75,12 +75,12 @@ inline FImportantLogScope FImportantLogScope::EnterImpl(uint32 Uid, uint32 Size)
 	uint8* Out = (uint8*)Buffer - Bias;
 
 	// Event header
-	auto* Header = (uint16*)Out;
-	Header[0] = uint16(Uid | uint32(EKnownEventUids::Flag_TwoByteUid));
-	Header[1] = uint16(Size);
+	auto* Header = (FImportantEventHeader*)Out;
+	Header->Uid = uint16(Uid);
+	Header->Size = uint16(Size);
 
 	FImportantLogScope Ret;
-	Ret.Ptr = (uint8*)(Header + 2);
+	Ret.Ptr = Header->Data;
 	Ret.BufferOffset = int32(PTRINT(Buffer) - PTRINT(Ret.Ptr));
 	Ret.AuxCursor = 0;
 	return Ret;
@@ -115,13 +115,14 @@ struct FImportantLogScope::FFieldSet<FieldMeta, Type[]>
 		uint32 Size = Num * sizeof(Type);
 
 		auto* Header = (FAuxHeader*)(Scope->Ptr + Scope->AuxCursor);
-		Header->Size = Size << 8;
-		Header->FieldIndex = uint8(0x80 | (FieldMeta::Index & int32(EIndexPack::NumFieldsMask)));
+		Header->Pack = Size << FAuxHeader::SizeShift;
+		Header->Pack |= (FieldMeta::Index & int32(EIndexPack::NumFieldsMask)) << FAuxHeader::FieldShift;
+		Header->Uid = uint8(EKnownEventUids::AuxData);
 
 		memcpy(Header + 1, Data, Size);
 
 		Scope->AuxCursor += sizeof(FAuxHeader) + Size;
-		Scope->Ptr[Scope->AuxCursor] = 0; // null terminator
+		Scope->Ptr[Scope->AuxCursor] = uint8(EKnownEventUids::AuxDataTerminal);
 	}
 };
 
@@ -137,13 +138,14 @@ struct FImportantLogScope::FFieldSet<FieldMeta, AnsiString>
 		}
 
 		auto* Header = (FAuxHeader*)(Scope->Ptr + Scope->AuxCursor);
-		Header->Size = Length << 8;
-		Header->FieldIndex = uint8(0x80 | (FieldMeta::Index & int32(EIndexPack::NumFieldsMask)));
+		Header->Pack = Length << FAuxHeader::SizeShift;
+		Header->Pack |= (FieldMeta::Index & int32(EIndexPack::NumFieldsMask)) << FAuxHeader::FieldShift;
+		Header->Uid = uint8(EKnownEventUids::AuxData);
 
 		memcpy(Header + 1, String, Length);
 
 		Scope->AuxCursor += sizeof(FAuxHeader) + Length;
-		Scope->Ptr[Scope->AuxCursor] = 0; // null terminator
+		Scope->Ptr[Scope->AuxCursor] = uint8(EKnownEventUids::AuxDataTerminal);
 	}
 
 	static void Impl(FImportantLogScope* Scope, const WIDECHAR* String, int32 Length=-1)
@@ -155,8 +157,9 @@ struct FImportantLogScope::FFieldSet<FieldMeta, AnsiString>
 		}
 
 		auto* Header = (FAuxHeader*)(Scope->Ptr + Scope->AuxCursor);
-		Header->Size = Length << 8;
-		Header->FieldIndex = uint8(0x80 | (FieldMeta::Index & int32(EIndexPack::NumFieldsMask)));
+		Header->Pack = Length << FAuxHeader::SizeShift;
+		Header->Pack |= (FieldMeta::Index & int32(EIndexPack::NumFieldsMask)) << FAuxHeader::FieldShift;
+		Header->Uid = uint8(EKnownEventUids::AuxData);
 
 		auto* Out = (int8*)(Header + 1);
 		for (int32 i = 0; i < Length; ++i)
@@ -167,7 +170,7 @@ struct FImportantLogScope::FFieldSet<FieldMeta, AnsiString>
 		}
 
 		Scope->AuxCursor += sizeof(FAuxHeader) + Length;
-		Scope->Ptr[Scope->AuxCursor] = 0; // null terminator
+		Scope->Ptr[Scope->AuxCursor] = uint8(EKnownEventUids::AuxDataTerminal);
 	}
 };
 
@@ -186,13 +189,14 @@ struct FImportantLogScope::FFieldSet<FieldMeta, WideString>
 		uint32 Size = Length * sizeof(WIDECHAR);
 
 		auto* Header = (FAuxHeader*)(Scope->Ptr + Scope->AuxCursor);
-		Header->Size = Size << 8;
-		Header->FieldIndex = uint8(0x80 | (FieldMeta::Index & int32(EIndexPack::NumFieldsMask)));
+		Header->Pack = Size << FAuxHeader::SizeShift;
+		Header->Pack |= (FieldMeta::Index & int32(EIndexPack::NumFieldsMask)) << FAuxHeader::FieldShift;
+		Header->Uid = uint8(EKnownEventUids::AuxData);
 
 		memcpy(Header + 1, String, Size);
 
 		Scope->AuxCursor += sizeof(FAuxHeader) + Size;
-		Scope->Ptr[Scope->AuxCursor] = 0; // null terminator
+		Scope->Ptr[Scope->AuxCursor] = uint8(EKnownEventUids::AuxDataTerminal);
 	}
 };
 
diff --git a/Engine/Source/Runtime/TraceLog/Public/Trace/Detail/LogScope.h b/Engine/Source/Runtime/TraceLog/Public/Trace/Detail/LogScope.h
--- a/Engine/Source/Runtime/TraceLog/Public/Trace/Detail/LogScope.h
+++ b/Engine/Source/Runtime/TraceLog/Public/Trace/Detail/LogScope.h
@@ -13,36 +13,48 @@ namespace Trace {
 namespace Private {
 
 struct FWriteBuffer;
+template <bool bMaybeHasAux> class TLogScope;
 
 ////////////////////////////////////////////////////////////////////////////////
 class FLogScope
 {
+	friend class FEventNode;
+
 public:
 	template <typename EventType>
-	static FLogScope		Enter(uint32 ExtraSize=0);
+	static auto				Enter();
 	template <typename EventType>
-	static FLogScope		ScopedEnter(uint32 ExtraSize=0);
+	static auto				ScopedEnter();
 	template <typename EventType>
-	static FLogScope		ScopedStampedEnter(uint32 ExtraSize=0);
+	static auto				ScopedStampedEnter();
 	void*					GetPointer() const			{ return Ptr; }
-	void					Commit() const;
-	void					operator += (const FLogScope&) const;
 	const FLogScope&		operator << (bool) const	{ return *this; }
 	constexpr explicit		operator bool () const		{ return true; }
 
 	template <typename FieldMeta, typename Type>
 	struct FFieldSet;
 
+protected:
+	void					Commit() const;
+
 private:
 	template <uint32 Flags>
-	static FLogScope		EnterImpl(uint32 Uid, uint32 Size);
-	template <class T> void	EnterPrelude(uint32 Size, bool bMaybeHasAux);
-	void					Enter(uint32 Uid, uint32 Size, bool bMaybeHasAux);
-	void					EnterNoSync(uint32 Uid, uint32 Size, bool bMaybeHasAux);
+	static auto				EnterImpl(uint32 Uid, uint32 Size);
+	template <class T> void	EnterPrelude(uint32 Size);
+	void					Enter(uint32 Uid, uint32 Size);
+	void					EnterNoSync(uint32 Uid, uint32 Size);
 	uint8*					Ptr;
 	FWriteBuffer*			Buffer;
 };
 
+template <bool bMaybeHasAux>
+class TLogScope
+	: public FLogScope
+{
+public:
+	void					operator += (const FLogScope&) const;
+};
+
 
 
 ////////////////////////////////////////////////////////////////////////////////
diff --git a/Engine/Source/Runtime/TraceLog/Public/Trace/Detail/LogScope.inl b/Engine/Source/Runtime/TraceLog/Public/Trace/Detail/LogScope.inl
--- a/Engine/Source/Runtime/TraceLog/Public/Trace/Detail/LogScope.inl
+++ b/Engine/Source/Runtime/TraceLog/Public/Trace/Detail/LogScope.inl
@@ -25,34 +25,27 @@ inline void FLogScope::Commit() const
 	AtomicStoreRelease((uint8**) &(Buffer->Committed), Buffer->Cursor);
 }
 
-////////////////////////////////////////////////////////////////////////////////
-inline void FLogScope::operator += (const FLogScope&) const
-{
-	Commit();
-}
-
 ////////////////////////////////////////////////////////////////////////////////
 template <uint32 Flags>
-inline FLogScope FLogScope::EnterImpl(uint32 Uid, uint32 Size)
+inline auto FLogScope::EnterImpl(uint32 Uid, uint32 Size)
 {
-	FLogScope Ret;
-	bool bMaybeHasAux = (Flags & FEventInfo::Flag_MaybeHasAux) != 0;
+	TLogScope<(Flags & FEventInfo::Flag_MaybeHasAux) != 0> Ret;
 	if ((Flags & FEventInfo::Flag_NoSync) != 0)
 	{
-		Ret.EnterNoSync(Uid, Size, bMaybeHasAux);
+		Ret.EnterNoSync(Uid, Size);
 	}
 	else
 	{
-		Ret.Enter(Uid, Size, bMaybeHasAux);
+		Ret.Enter(Uid, Size);
 	}
 	return Ret;
 }
 
 ////////////////////////////////////////////////////////////////////////////////
 template <class HeaderType>
-inline void FLogScope::EnterPrelude(uint32 Size, bool bMaybeHasAux)
+inline void FLogScope::EnterPrelude(uint32 Size)
 {
-	uint32 AllocSize = sizeof(HeaderType) + Size + int32(bMaybeHasAux);
+	uint32 AllocSize = sizeof(HeaderType) + Size;
 
 	Buffer = Writer_GetBuffer();
 	Buffer->Cursor += AllocSize;
@@ -61,40 +54,53 @@ inline void FLogScope::EnterPrelude(uint32 Size, bool bMaybeHasAux)
 		Buffer = Writer_NextBuffer(AllocSize);
 	}
 
-	// The auxilary data null terminator.
-	if (bMaybeHasAux)
-	{
-		Buffer->Cursor[-1] = 0;
-	}
-
-	Ptr = Buffer->Cursor - Size - int32(bMaybeHasAux);
+	Ptr = Buffer->Cursor - Size;
 }
 
 ////////////////////////////////////////////////////////////////////////////////
-inline void FLogScope::Enter(uint32 Uid, uint32 Size, bool bMaybeHasAux)
+inline void FLogScope::Enter(uint32 Uid, uint32 Size)
 {
-	EnterPrelude<FEventHeaderSync>(Size, bMaybeHasAux);
+	EnterPrelude<FEventHeaderSync>(Size);
 
 	// Event header
 	auto* Header = (uint16*)(Ptr - sizeof(FEventHeaderSync::SerialHigh));
 	*(uint32*)(Header - 1) = uint32(AtomicAddRelaxed(&GLogSerial, 1u));
-	Header[-2] = uint16(Size);
-	Header[-3] = uint16(Uid)|int32(EKnownEventUids::Flag_TwoByteUid);
+	Header[-2] = uint16(Uid)|int32(EKnownEventUids::Flag_TwoByteUid);
 }
 
 ////////////////////////////////////////////////////////////////////////////////
-inline void FLogScope::EnterNoSync(uint32 Uid, uint32 Size, bool bMaybeHasAux)
+inline void FLogScope::EnterNoSync(uint32 Uid, uint32 Size)
 {
-	EnterPrelude<FEventHeader>(Size, bMaybeHasAux);
+	EnterPrelude<FEventHeader>(Size);
 
 	// Event header
 	auto* Header = (uint16*)(Ptr);
-	Header[-1] = uint16(Size);
-	Header[-2] = uint16(Uid)|int32(EKnownEventUids::Flag_TwoByteUid);
+	Header[-1] = uint16(Uid)|int32(EKnownEventUids::Flag_TwoByteUid);
+}
+
+
+
+////////////////////////////////////////////////////////////////////////////////
+template </*bMaybeHasAux*/>
+inline void TLogScope<false>::operator += (const FLogScope&) const
+{
+	Commit();
+}
+
+////////////////////////////////////////////////////////////////////////////////
+template </*bMaybeHasAux*/>
+inline void TLogScope<true>::operator += (const FLogScope&) const
+{
+	FWriteBuffer* LatestBuffer = Writer_GetBuffer();
+	LatestBuffer->Cursor[0] = uint8(EKnownEventUids::AuxDataTerminal << EKnownEventUids::_UidShift);
+	LatestBuffer->Cursor++;
+
+	Commit();
 }
 
 
 
+
 ////////////////////////////////////////////////////////////////////////////////
 inline FScopedLogScope::~FScopedLogScope()
 {
@@ -103,7 +109,7 @@ inline FScopedLogScope::~FScopedLogScope()
 		return;
 	}
 
-	uint8 LeaveUid = uint8(EKnownEventUids::LeaveScope) << EKnownEventUids::_UidShift;
+	uint8 LeaveUid = uint8(EKnownEventUids::LeaveScope << EKnownEventUids::_UidShift);
 
 	FWriteBuffer* Buffer = Writer_GetBuffer();
 	if (UNLIKELY(int32((uint8*)Buffer - Buffer->Cursor)) < int32(sizeof(LeaveUid)))
@@ -159,19 +165,19 @@ inline void FScopedStampedLogScope::SetActive()
 
 
 ////////////////////////////////////////////////////////////////////////////////
-template <class T>
-FORCENOINLINE FLogScope FLogScope::Enter(uint32 ExtraSize)
+template <class EventType>
+FORCENOINLINE auto FLogScope::Enter()
 {
-	uint32 Size = T::GetSize() + ExtraSize;
-	uint32 Uid = T::GetUid();
-	return EnterImpl<T::EventFlags>(Uid, Size);
+	uint32 Size = EventType::GetSize();
+	uint32 Uid = EventType::GetUid();
+	return EnterImpl<EventType::EventFlags>(Uid, Size);
 }
 
 ////////////////////////////////////////////////////////////////////////////////
-template <class T>
-FORCENOINLINE FLogScope FLogScope::ScopedEnter(uint32 ExtraSize)
+template <class EventType>
+FORCENOINLINE auto FLogScope::ScopedEnter()
 {
-	uint8 EnterUid = uint8(EKnownEventUids::EnterScope) << EKnownEventUids::_UidShift;
+	uint8 EnterUid = uint8(EKnownEventUids::EnterScope << EKnownEventUids::_UidShift);
 
 	FWriteBuffer* Buffer = Writer_GetBuffer();
 	if (UNLIKELY(int32((uint8*)Buffer - Buffer->Cursor)) < int32(sizeof(EnterUid)))
@@ -184,12 +190,12 @@ FORCENOINLINE FLogScope FLogScope::ScopedEnter(uint32 ExtraSize)
 
 	AtomicStoreRelease((uint8**) &(Buffer->Committed), Buffer->Cursor);
 
-	return Enter<T>(ExtraSize);
+	return Enter<EventType>();
 }
 
 ////////////////////////////////////////////////////////////////////////////////
-template <class T>
-FORCENOINLINE FLogScope FLogScope::ScopedStampedEnter(uint32 ExtraSize)
+template <class EventType>
+FORCENOINLINE auto FLogScope::ScopedStampedEnter()
 {
 	uint64 Stamp;
 
@@ -207,7 +213,7 @@ FORCENOINLINE FLogScope FLogScope::ScopedStampedEnter(uint32 ExtraSize)
 
 	AtomicStoreRelease((uint8**) &(Buffer->Committed), Buffer->Cursor);
 
-	return Enter<T>(ExtraSize);
+	return Enter<EventType>();
 }
 
 
@@ -229,7 +235,7 @@ struct FLogScope::FFieldSet<FieldMeta, Type[]>
 {
 	static void Impl(FLogScope*, Type const* Data, int32 Num)
 	{
-		static const uint32 Index = FieldMeta::Index & 0x7f;
+		static const uint32 Index = FieldMeta::Index & int32(EIndexPack::NumFieldsMask);
 		int32 Size = (Num * sizeof(Type)) & (FAuxHeader::SizeLimit - 1) & ~(sizeof(Type) - 1);
 		Field_WriteAuxData(Index, (const uint8*)Data, Size);
 	}
@@ -255,7 +261,7 @@ struct FLogScope::FFieldSet<FieldMeta, AnsiString>
 			Length = int32(strlen(String));
 		}
 
-		static const uint32 Index = FieldMeta::Index & 0x7f;
+		static const uint32 Index = FieldMeta::Index & int32(EIndexPack::NumFieldsMask);
 		Field_WriteStringAnsi(Index, String, Length);
 	}
 
@@ -267,7 +273,7 @@ struct FLogScope::FFieldSet<FieldMeta, AnsiString>
 			for (const WIDECHAR* c = String; *c; ++c, ++Length);
 		}
 
-		static const uint32 Index = FieldMeta::Index & 0x7f;
+		static const uint32 Index = FieldMeta::Index & int32(EIndexPack::NumFieldsMask);
 		Field_WriteStringAnsi(Index, String, Length);
 	}
 };
@@ -284,7 +290,7 @@ struct FLogScope::FFieldSet<FieldMeta, WideString>
 			for (const WIDECHAR* c = String; *c; ++c, ++Length);
 		}
 
-		static const uint32 Index = FieldMeta::Index & 0x7f;
+		static const uint32 Index = FieldMeta::Index & int32(EIndexPack::NumFieldsMask);
 		Field_WriteStringWide(Index, String, Length);
 	}
 };
diff --git a/Engine/Source/Runtime/TraceLog/Public/Trace/Detail/Protocol.h b/Engine/Source/Runtime/TraceLog/Public/Trace/Detail/Protocol.h
--- a/Engine/Source/Runtime/TraceLog/Public/Trace/Detail/Protocol.h
+++ b/Engine/Source/Runtime/TraceLog/Public/Trace/Detail/Protocol.h
@@ -6,7 +6,7 @@
 
 ////////////////////////////////////////////////////////////////////////////////
 #if UE_TRACE_ENABLED
-#	define TRACE_PRIVATE_PROTOCOL_4
+#	define TRACE_PRIVATE_PROTOCOL_5
 #endif
 
 #if defined(_MSC_VER)
@@ -19,6 +19,7 @@
 #include "Protocols/Protocol2.h"
 #include "Protocols/Protocol3.h"
 #include "Protocols/Protocol4.h"
+#include "Protocols/Protocol5.h"
 
 #if defined(_MSC_VER)
 	#pragma warning(pop)
diff --git a/Engine/Source/Runtime/TraceLog/Public/Trace/Detail/Protocols/Protocol5.h b/Engine/Source/Runtime/TraceLog/Public/Trace/Detail/Protocols/Protocol5.h
new file mode 100644
--- /dev/null
+++ b/Engine/Source/Runtime/TraceLog/Public/Trace/Detail/Protocols/Protocol5.h
@@ -0,0 +1,109 @@
+// Copyright Epic Games, Inc. All Rights Reserved.
+
+#pragma once
+
+namespace UE {
+namespace Trace {
+
+#if defined(TRACE_PRIVATE_PROTOCOL_5)
+inline
+#endif
+namespace Protocol5
+{
+
+////////////////////////////////////////////////////////////////////////////////
+enum EProtocol : uint8 { Id = 5 };
+
+////////////////////////////////////////////////////////////////////////////////
+using Protocol4::EFieldType;
+using Protocol4::FNewEventEvent;
+using Protocol4::EEventFlags;
+
+////////////////////////////////////////////////////////////////////////////////
+struct EKnownEventUids
+{
+	static const uint16 Flag_TwoByteUid	= 1 << 0;
+	static const uint16 _UidShift		= 1;
+	enum : uint16
+	{
+		NewEvent						= 0,
+		AuxData,
+		_AuxData_Unused,
+		AuxDataTerminal,
+		EnterScope,
+		LeaveScope,
+		_Unused6,
+		_Unused7,
+		EnterScope_T,
+		_EnterScope_T_Unused0,	// reserved for variable
+		_EnterScope_T_Unused1,	// length timestamps
+		_EnterScope_T_Unused2,
+		LeaveScope_T,
+		_LeaveScope_T_Unused0,
+		_LeaveScope_T_Unused1,
+		_LeaveScope_T_Unused2,
+		_WellKnownNum,
+	};
+	static const uint16 User			= _WellKnownNum;
+	static const uint16 Max				= (1 << (16 - _UidShift)) - 1;
+	static const uint16 Invalid			= Max;
+};
+
+////////////////////////////////////////////////////////////////////////////////
+struct FEventHeader
+{
+	uint16		Uid;
+	uint8		Data[];
+};
+static_assert(sizeof(FEventHeader) == 2, "Struct layout assumption doesn't match expectation");
+
+////////////////////////////////////////////////////////////////////////////////
+struct FImportantEventHeader
+{
+	uint16		Uid;
+	uint16		Size;
+	uint8		Data[];
+};
+static_assert(sizeof(FImportantEventHeader) == 4, "Struct layout assumption doesn't match expectation");
+
+////////////////////////////////////////////////////////////////////////////////
+#pragma pack(push, 1)
+struct FEventHeaderSync
+{
+	uint16		Uid;
+	uint16		SerialLow;		// 24-bit
+	uint8		SerialHigh;		// serial no.
+	uint8		Data[];
+};
+#pragma pack(pop)
+static_assert(sizeof(FEventHeaderSync) == 5, "Packing assumption doesn't hold");
+
+////////////////////////////////////////////////////////////////////////////////
+struct FAuxHeader
+{
+	enum : uint32
+	{
+		FieldShift	= 8,
+		FieldBits	= 5,
+		FieldMask	= (1 << FieldBits) - 1,
+		SizeShift	= FieldShift + FieldBits,
+		SizeLimit	= 1 << (32 - SizeShift),
+	};
+
+	union
+	{
+		struct
+		{
+			uint8	Uid;
+			uint8	FieldIndex_Size;
+			uint16	Size;
+		};
+		uint32		Pack;
+	};
+	uint8		Data[];
+};
+static_assert(sizeof(FAuxHeader) == 4, "Struct layout assumption doesn't match expectation");
+
+} // namespace Protocol5
+} // namespace Trace
+} // namespace UE
diff --git a/Engine/Source/Runtime/TraceLog/Public/Trace/Detail/Writer.inl b/Engine/Source/Runtime/TraceLog/Public/Trace/Detail/Writer.inl
--- a/Engine/Source/Runtime/TraceLog/Public/Trace/Detail/Writer.inl
+++ b/Engine/Source/Runtime/TraceLog/Public/Trace/Detail/Writer.inl
@@ -11,9 +11,8 @@ namespace Private {
 ////////////////////////////////////////////////////////////////////////////////
 struct FWriteBuffer
 {
-	uint32						Overflow;
+	uint8						Overflow[6];
 	uint16						Size;
-	uint16						ThreadId;
 	uint64						PrevTimestamp;
 	FWriteBuffer* __restrict	NextThread;
 	FWriteBuffer* __restrict	NextBuffer;
@@ -21,7 +20,8 @@ struct FWriteBuffer
 	uint8* __restrict volatile	Committed;
 	uint8* __restrict			Reaped;
 	int32 volatile				EtxOffset;
-	int32						Partial;
+	int16						Partial;
+	uint16						ThreadId;
 };
 
 ////////////////////////////////////////////////////////////////////////////////
