Restrict trace store client to only resolve IPv4 addresses. Previously the resolve would look up IPv6 addresses for the given host leading to stalls.

#rnx

diff --git a/Engine/Source/Developer/TraceAnalysis/Private/Store/StoreClient.cpp b/Engine/Source/Developer/TraceAnalysis/Private/Store/StoreClient.cpp
--- a/Engine/Source/Developer/TraceAnalysis/Private/Store/StoreClient.cpp
+++ b/Engine/Source/Developer/TraceAnalysis/Private/Store/StoreClient.cpp
@@ -189,7 +189,10 @@ bool FStoreCborClient::Connect(const TCHAR* Host, uint16 Port)
 	FCStringAnsi::Sprintf(PortString, "%d", Port);
 
 	asio::ip::tcp::resolver Resolver(IoContext);
-	asio::ip::tcp::resolver::results_type Endpoints = Resolver.resolve((const char*)HostUtf8.Get(), PortString);
+	asio::ip::tcp::resolver::results_type Endpoints = Resolver.resolve(
+		asio::ip::tcp::resolver::protocol_type::v4(),
+		(const char*)HostUtf8.Get(),
+		PortString);
 
 	asio::error_code ErrorCode;
 #if PLATFORM_WINDOWS
