Fixed dispatch of events with auxiliary data for the new protocol version (5).
Added bounds checking when reading (skipping) the metadata in front of trace file.

diff --git a/Engine/Source/Developer/TraceAnalysis/Private/Analysis/Engine.cpp b/Engine/Source/Developer/TraceAnalysis/Private/Analysis/Engine.cpp
--- a/Engine/Source/Developer/TraceAnalysis/Private/Analysis/Engine.cpp
+++ b/Engine/Source/Developer/TraceAnalysis/Private/Analysis/Engine.cpp
@@ -211,9 +211,9 @@ const ANSICHAR* IAnalyzer::FThreadInfo::GetGroupName() const
 ////////////////////////////////////////////////////////////////////////////////
 struct FTiming
 {
-	uint64	BaseTimestamp;
-	uint64	TimestampHz;
-	double	InvTimestampHz;
+	uint64	BaseTimestamp = 0;
+	uint64	TimestampHz = 0;
+	double	InvTimestampHz = 0.0;
 	uint64	EventTimestamp = 0;
 };
 
@@ -2457,11 +2457,18 @@ public:
 private:
 	struct alignas(16) FEventDesc
 	{
-		int32				Serial;
-		uint16				Uid			: 15;
-		uint16				bHasAux		: 1;
-		uint16				AuxKey;
-		const uint8*		Data;
+		union
+		{
+			struct
+			{
+				int32		Serial;
+				uint16		Uid			: 15;
+				uint16		bHasAux		: 1;
+				uint16		AuxKey;
+			};
+			uint64			Meta = 0;
+		};
+		const uint8*		Data = nullptr;
 	};
 	static_assert(sizeof(FEventDesc) == 16, "");
 
@@ -2602,9 +2609,9 @@ int32 FProtocol5Stage::ParseImportantEvents(FStreamReader& Reader, EventDescArra
 		uint32 Uid = Header->Uid;
 
 		FEventDesc EventDesc;
+		EventDesc.Serial = ESerial::Ignored;
 		EventDesc.Uid = Uid;
 		EventDesc.Data = Header->Data;
-		EventDesc.Serial = ESerial::Ignored;
 
 		// Special case for new events. It would work to add a 0 type to the
 		// registry but this way avoid raveling things together.
@@ -2622,12 +2629,15 @@ int32 FProtocol5Stage::ParseImportantEvents(FStreamReader& Reader, EventDescArra
 			return 1;
 		}
 
-		OutEventDescs.Add(EventDesc);
-
 		if (TypeInfo->Flags & FTypeRegistry::FTypeInfo::Flag_MaybeHasAux)
 		{
 			EventDesc.bHasAux = 1;
+		}
+
+		OutEventDescs.Add(EventDesc);
 
+		if (TypeInfo->Flags & FTypeRegistry::FTypeInfo::Flag_MaybeHasAux)
+		{
 			const uint8* Cursor = Header->Data + TypeInfo->EventSize;
 			const uint8* End = Header->Data + Header->Size;
 			while (Cursor <= End)
@@ -3168,6 +3178,12 @@ FMetadataStage::EStatus FMetadataStage::OnData(
 		return EStatus::NotEnoughData;
 	}
 
+	const uint8* Metadata = Reader.GetPointer(sizeof(*MetadataSize) + *MetadataSize);
+	if (Metadata == nullptr)
+	{
+		return EStatus::NotEnoughData;
+	}
+
 	Reader.Advance(sizeof(*MetadataSize) + *MetadataSize);
 
 	Context.Machine.QueueStage<FEstablishTransportStage>();
