[Insights] Fixed the processing of scoped timing events with 0 duration.

Description of the problem: Very short scopes for cpu timing events with metadata can have a duration of 0. In this case the enter/leave events are incorrectly processed. The "enter" scope event is processed as it has a relative timestamp > 0, but the "leave" event is ignored in analysis as it will have a timestamp == 0, some as scoped events without time. It results in a broken hierarchy of cpu timing events!

The fix ensures the scoped events with 0 duration will still have a valid timestamp (in this case the "leave" event will have same timestamp as "enter" event, as expected), so those "leave scope" events are not ignored anymore by the cpu profiler analyzer (which still needs to ignore scoped events without time).

diff --git a/Engine/Source/Developer/TraceAnalysis/Private/Analysis/Engine.cpp b/Engine/Source/Developer/TraceAnalysis/Private/Analysis/Engine.cpp
--- a/Engine/Source/Developer/TraceAnalysis/Private/Analysis/Engine.cpp
+++ b/Engine/Source/Developer/TraceAnalysis/Private/Analysis/Engine.cpp
@@ -1465,8 +1465,13 @@ public:
 	void				OnNewType(const FTypeRegistry::FTypeInfo* TypeInfo);
 	void				OnEvent(const FEventDataInfo& EventDataInfo);
 	virtual void		OnThreadInfo(const FThreads::FInfo& InThreadInfo) override;
-	void				EnterScope(uint64 Timestamp=0);
-	void				LeaveScope(uint64 Timestamp=0);
+	void				EnterScope();
+	void				EnterScope(uint64 Timestamp);
+	void				LeaveScope();
+	void				LeaveScope(uint64 Timestamp);
+
+private:
+	void				DispatchLeaveScope();
 
 private:
 	FAnalysisState		State;
@@ -1552,27 +1557,38 @@ void FAnalysisBridge::OnThreadInfo(const FThreads::FInfo& InThreadInfo)
 	AnalyzerHub.OnThreadInfo(InThreadInfo);
 }
 
+////////////////////////////////////////////////////////////////////////////////
+void FAnalysisBridge::EnterScope()
+{
+	ThreadInfo->ScopeRoutes.Push(~0);
+}
+
 ////////////////////////////////////////////////////////////////////////////////
 void FAnalysisBridge::EnterScope(uint64 Timestamp)
 {
-    if (Timestamp)
-    {
-	    Timestamp = ThreadInfo->PrevTimestamp += Timestamp;
-    }
+	Timestamp = ThreadInfo->PrevTimestamp += Timestamp;
+	ThreadInfo->ScopeRoutes.Push(~Timestamp);
+}
 
-    ThreadInfo->ScopeRoutes.Push(~Timestamp);
+////////////////////////////////////////////////////////////////////////////////
+void FAnalysisBridge::LeaveScope()
+{
+	State.Timing.EventTimestamp = 0;
+	DispatchLeaveScope();
 }
 
 ////////////////////////////////////////////////////////////////////////////////
 void FAnalysisBridge::LeaveScope(uint64 Timestamp)
 {
-    if (Timestamp)
-    {
-	    Timestamp = ThreadInfo->PrevTimestamp += Timestamp;
-    }
-
+	Timestamp = ThreadInfo->PrevTimestamp += Timestamp;
 	State.Timing.EventTimestamp = Timestamp;
+	DispatchLeaveScope();
+	State.Timing.EventTimestamp = 0;
+}
 
+////////////////////////////////////////////////////////////////////////////////
+void FAnalysisBridge::DispatchLeaveScope()
+{
 	if (ThreadInfo->ScopeRoutes.Num() > 0)
 	{
 		PTRINT TypeInfoPtr = ThreadInfo->ScopeRoutes.Pop(false);
@@ -1591,8 +1607,6 @@ void FAnalysisBridge::LeaveScope(uint64 Timestamp)
 
 		AnalyzerHub.OnEvent(*TypeInfo, IAnalyzer::EStyle::LeaveScope, Context);
 	}
-
-	State.Timing.EventTimestamp = 0;
 }
 
 
