Added a backtrack call to be able to restore a reader to an earlier point in its buffer. Like marks, only for the times when one wasn't able to take a mark.

#rnx

diff --git a/Engine/Source/Developer/TraceAnalysis/Private/Analysis/StreamReader.cpp b/Engine/Source/Developer/TraceAnalysis/Private/Analysis/StreamReader.cpp
--- a/Engine/Source/Developer/TraceAnalysis/Private/Analysis/StreamReader.cpp
+++ b/Engine/Source/Developer/TraceAnalysis/Private/Analysis/StreamReader.cpp
@@ -43,6 +43,19 @@ bool FStreamReader::IsEmpty() const
 	return Cursor >= End;
 }
 
+////////////////////////////////////////////////////////////////////////////////
+bool FStreamReader::Backtrack(const uint8* To)
+{
+	uint32 BacktrackedCursor = uint32(UPTRINT(To - Buffer));
+	if (BacktrackedCursor > End)
+	{
+		return false;
+	}
+
+	Cursor = BacktrackedCursor;
+	return true;
+}
+
 ////////////////////////////////////////////////////////////////////////////////
 struct FMark* FStreamReader::SaveMark() const
 {
diff --git a/Engine/Source/Developer/TraceAnalysis/Private/Analysis/StreamReader.h b/Engine/Source/Developer/TraceAnalysis/Private/Analysis/StreamReader.h
--- a/Engine/Source/Developer/TraceAnalysis/Private/Analysis/StreamReader.h
+++ b/Engine/Source/Developer/TraceAnalysis/Private/Analysis/StreamReader.h
@@ -21,6 +21,7 @@ public:
 	void						Advance(uint32 Size);
 	bool						IsEmpty() const;
 	int32						GetRemaining() const;
+	bool						Backtrack(const uint8* To);
 	struct FMark*				SaveMark() const;
 	void						RestoreMark(struct FMark* Mark);
 
