Consume unsynchronised events before sorting so it is only ever serial numbers that define the heap.

#rnx

diff --git a/Engine/Source/Developer/TraceAnalysis/Private/Analysis/Engine.cpp b/Engine/Source/Developer/TraceAnalysis/Private/Analysis/Engine.cpp
--- a/Engine/Source/Developer/TraceAnalysis/Private/Analysis/Engine.cpp
+++ b/Engine/Source/Developer/TraceAnalysis/Private/Analysis/Engine.cpp
@@ -2727,16 +2727,46 @@ FProtocol5Stage::EStatus FProtocol5Stage::OnDataNormal(const FMachineContext& Co
 		TRACE_ANALYSIS_DEBUG("Thread: %03d bNotEnoughData:%d", i, bNotEnoughData);
 	}
 
-	// Early out if there isn't any events available.
-	if (UNLIKELY(EventDescHeap.IsEmpty()))
+	// Now EventDescs is stable we can convert the indices into pointers
+	for (FEventDescStream& Stream : EventDescHeap)
 	{
-		return bNotEnoughData ? EStatus::NotEnoughData : EStatus::Eof;
+		Stream.EventDescs = EventDescs.GetData() + Stream.Index;
 	}
 
-	// Now EventDescs is stable we can convert the indices into pointers
+	// Process leading unsynchronised events so that each stream starts with a
+	// sychronised event.
 	for (FEventDescStream& Stream : EventDescHeap)
 	{
-		Stream.EventDescs = EventDescs.GetData() + Stream.Index;
+		// Extract a run of consecutive unsynchronised events
+		const FEventDesc* EndDesc = Stream.EventDescs;
+		for (; EndDesc->Serial == ESerial::Ignored; ++EndDesc);
+
+		// Dispatch.
+		const FEventDesc* StartDesc = Stream.EventDescs;
+		int32 DescNum = int32(UPTRINT(EndDesc - StartDesc));
+		if (DescNum > 0)
+		{
+			Context.Bridge.SetActiveThread(Stream.ThreadId);
+
+			if (DispatchEvents(Context.Bridge, StartDesc, DescNum) < 0)
+			{
+				return EStatus::Error;
+			}
+
+			Stream.EventDescs = EndDesc;
+		}
+	}
+
+	// Trim off empty streams
+	EventDescHeap.RemoveAllSwap([] (const FEventDescStream& Stream)
+	{
+		return (Stream.EventDescs->Serial == ESerial::Terminal);
+	});
+
+	// Early out if there isn't any events available.
+	if (UNLIKELY(EventDescHeap.IsEmpty()))
+	{
+		return bNotEnoughData ? EStatus::NotEnoughData : EStatus::Eof;
 	}
 
 	// Provided that less than approximately "SerialRange * BytesPerSerial"
