Added an alignment to hint to the compiler to do the right thing.

#rnx

diff --git a/Engine/Source/Developer/TraceAnalysis/Private/Analysis/Engine.cpp b/Engine/Source/Developer/TraceAnalysis/Private/Analysis/Engine.cpp
--- a/Engine/Source/Developer/TraceAnalysis/Private/Analysis/Engine.cpp
+++ b/Engine/Source/Developer/TraceAnalysis/Private/Analysis/Engine.cpp
@@ -2691,7 +2691,7 @@ FProtocol5Stage::EStatus FProtocol5Stage::OnDataNormal(const FMachineContext& Co
 	EventDescs.Reset();
 	bool bNotEnoughData = false;
 
-	struct FEventDescStream
+	struct alignas(16) FEventDescStream
 	{
 		uint32				ThreadId;
 		uint32				Index;
