Set serial to "ignored" for aux data and terminal events.

#rnx

diff --git a/Engine/Source/Developer/TraceAnalysis/Private/Analysis/Engine.cpp b/Engine/Source/Developer/TraceAnalysis/Private/Analysis/Engine.cpp
--- a/Engine/Source/Developer/TraceAnalysis/Private/Analysis/Engine.cpp
+++ b/Engine/Source/Developer/TraceAnalysis/Private/Analysis/Engine.cpp
@@ -2666,6 +2666,7 @@ int32 FProtocol5Stage::ParseImportantEvents(FStreamReader& Reader, EventDescArra
 				FEventDesc& AuxDesc = OutEventDescs.Emplace_GetRef();
 				AuxDesc.Uid = uint8(EKnownUids::AuxData);
 				AuxDesc.Data = AuxHeader->Data;
+				AuxDesc.Serial = ESerial::Ignored;
 
 				Cursor = AuxHeader->Data + (AuxHeader->Pack >> FAuxHeader::SizeShift);
 			}
@@ -2892,6 +2893,7 @@ int32 FProtocol5Stage::ParseEventsWithAux(FStreamReader& Reader, EventDescArray&
 	while (!Reader.IsEmpty())
 	{
 		FEventDesc EventDesc;
+		EventDesc.Serial = ESerial::Ignored;
 
 		int32 Size = ParseEvent(Reader, EventDesc);
 		if (Size <= 0)
