Keep track of which reader in the transport a stream came from.

#rnx

diff --git a/Engine/Source/Developer/TraceAnalysis/Private/Analysis/Engine.cpp b/Engine/Source/Developer/TraceAnalysis/Private/Analysis/Engine.cpp
--- a/Engine/Source/Developer/TraceAnalysis/Private/Analysis/Engine.cpp
+++ b/Engine/Source/Developer/TraceAnalysis/Private/Analysis/Engine.cpp
@@ -2694,9 +2694,13 @@ FProtocol5Stage::EStatus FProtocol5Stage::OnDataNormal(const FMachineContext& Co
 
 	struct alignas(16) FEventDescStream
 	{
-		uint32				ThreadId;
-		uint32				Index;
-		const FEventDesc*	EventDescs;
+		uint32					ThreadId;
+		uint32					TransportIndex;
+		union
+		{
+			uint32				Index;
+			const FEventDesc*	EventDescs;
+		};
 	};
 	TArray<FEventDescStream> EventDescHeap;
 	EventDescHeap.Reserve(Transport.GetThreadCount());
@@ -2722,7 +2726,11 @@ FProtocol5Stage::EStatus FProtocol5Stage::OnDataNormal(const FMachineContext& Co
 			FEventDesc& EventDesc = EventDescs.Emplace_GetRef();
 			EventDesc.Serial = ESerial::Terminal;
 
-			EventDescHeap.Add({Transport.GetThreadId(i), NumEventDescs});
+			FEventDescStream Out;
+			Out.ThreadId = Transport.GetThreadId(i);
+			Out.TransportIndex = i;
+			Out.Index = NumEventDescs;
+			EventDescHeap.Add(Out);
 		}
 
 		TRACE_ANALYSIS_DEBUG("Thread: %03d bNotEnoughData:%d", i, bNotEnoughData);
@@ -2813,7 +2821,8 @@ FProtocol5Stage::EStatus FProtocol5Stage::OnDataNormal(const FMachineContext& Co
 		// Update the heap
 		if (EndDesc->Serial != ESerial::Terminal)
 		{
-			EventDescHeap.Add({Stream.ThreadId, 0, EndDesc});
+			auto& Out = EventDescHeap.Add_GetRef({Stream.ThreadId, Stream.TransportIndex});
+			Out.EventDescs = EndDesc;
 		}
 		EventDescHeap.HeapPopDiscard(Comparison);
 
