Tolerance against situations where there was no trace event in between and enter/leave scope pair

#rnx

diff --git a/Engine/Source/Developer/TraceAnalysis/Private/Analysis/Engine.cpp b/Engine/Source/Developer/TraceAnalysis/Private/Analysis/Engine.cpp
--- a/Engine/Source/Developer/TraceAnalysis/Private/Analysis/Engine.cpp
+++ b/Engine/Source/Developer/TraceAnalysis/Private/Analysis/Engine.cpp
@@ -1472,8 +1472,6 @@ public:
 
 private:
 	void				DispatchLeaveScope();
-
-private:
 	FAnalysisState		State;
 	FTraceAnalyzer		TraceAnalyzer = { State, *this };
 	FAnalyzerHub		AnalyzerHub;
@@ -1589,24 +1587,33 @@ void FAnalysisBridge::LeaveScope(uint64 Timestamp)
 ////////////////////////////////////////////////////////////////////////////////
 void FAnalysisBridge::DispatchLeaveScope()
 {
-	if (ThreadInfo->ScopeRoutes.Num() > 0)
+	if (ThreadInfo->ScopeRoutes.Num() <= 0)
 	{
-		PTRINT TypeInfoPtr = ThreadInfo->ScopeRoutes.Pop(false);
-		const auto* TypeInfo = (FTypeRegistry::FTypeInfo*)TypeInfoPtr;
+		// Leave scope without a corresponding enter
+		return;
+	}
 
-		FEventDataInfo EmptyEventInfo = {
-			nullptr,
-			*TypeInfo
-		};
+	int64 ScopeValue = int64(ThreadInfo->ScopeRoutes.Pop(false));
+	if (ScopeValue < 0)
+	{
+		// enter/leave pair without an event inbetween.
+		return;
+	}
 
-		IAnalyzer::FOnEventContext Context = {
-			*(const IAnalyzer::FThreadInfo*)ThreadInfo,
-			(const IAnalyzer::FEventTime&)(State.Timing),
-			(const IAnalyzer::FEventData&)EmptyEventInfo,
-		};
+	const auto* TypeInfo = (FTypeRegistry::FTypeInfo*)PTRINT(ScopeValue);
 
-		AnalyzerHub.OnEvent(*TypeInfo, IAnalyzer::EStyle::LeaveScope, Context);
-	}
+	FEventDataInfo EmptyEventInfo = {
+		nullptr,
+		*TypeInfo
+	};
+
+	IAnalyzer::FOnEventContext Context = {
+		*(const IAnalyzer::FThreadInfo*)ThreadInfo,
+		(const IAnalyzer::FEventTime&)(State.Timing),
+		(const IAnalyzer::FEventData&)EmptyEventInfo,
+	};
+
+	AnalyzerHub.OnEvent(*TypeInfo, IAnalyzer::EStyle::LeaveScope, Context);
 }
 
 
