Keep track of whether an event had a one or two byte event UID

#rnx

diff --git a/Engine/Source/Developer/TraceAnalysis/Private/Analysis/Engine.cpp b/Engine/Source/Developer/TraceAnalysis/Private/Analysis/Engine.cpp
--- a/Engine/Source/Developer/TraceAnalysis/Private/Analysis/Engine.cpp
+++ b/Engine/Source/Developer/TraceAnalysis/Private/Analysis/Engine.cpp
@@ -2483,7 +2483,8 @@ private:
 			struct
 			{
 				int32		Serial;
-				uint16		Uid			: 15;
+				uint16		Uid			: 14;
+				uint16		bTwoByteUid	: 1;
 				uint16		bHasAux		: 1;
 				uint16		AuxKey;
 			};
@@ -2992,6 +2993,7 @@ int32 FProtocol5Stage::ParseEvent(FStreamReader& Reader, FEventDesc& EventDesc)
 	uint32 Uid = *Cursor;
 	if (Uid & EKnownUids::Flag_TwoByteUid)
 	{
+		EventDesc.bTwoByteUid = 1;
 		Uid = *(uint16*)Cursor;
 		++Cursor;
 	}
