Always expect events to be contiguous.

Previously it was assumed that the earliest known event was a good starting point, but in reality analysis never has a complete view of the world to be able to make any decisions about the data it is receiving. New data could easily arrive that invalidates any assumptions. Now nothing is assumed. Proceed only when one particular event is encountered. Note that is not quite complete - it doesn't cover holes that can occur early in the data stream.

#rnx

diff --git a/Engine/Source/Developer/TraceAnalysis/Private/Analysis/Engine.cpp b/Engine/Source/Developer/TraceAnalysis/Private/Analysis/Engine.cpp
--- a/Engine/Source/Developer/TraceAnalysis/Private/Analysis/Engine.cpp
+++ b/Engine/Source/Developer/TraceAnalysis/Private/Analysis/Engine.cpp
@@ -2519,6 +2519,7 @@ private:
 	FTypeRegistry			TypeRegistry;
 	FTidPacketTransport&	Transport;
 	EventDescArray			EventDescs;
+	uint32					NextSerial = ~0u;
 };
 
 ////////////////////////////////////////////////////////////////////////////////
@@ -2800,6 +2801,12 @@ FProtocol5Stage::EStatus FProtocol5Stage::OnDataNormal(const FMachineContext& Co
 	};
 	EventDescHeap.Heapify(Comparison);
 
+	// Events must be consumed contiguously.
+	if (NextSerial == ~0u)
+	{
+		NextSerial = EventDescHeap.HeapTop().EventDescs[0].Serial;
+	}
+
 	do
 	{
 		const FEventDescStream& Stream = EventDescHeap.HeapTop();
@@ -2809,22 +2816,26 @@ FProtocol5Stage::EStatus FProtocol5Stage::OnDataNormal(const FMachineContext& Co
 		Context.Bridge.SetActiveThread(Stream.ThreadId);
 
 		// Extract a run of consecutive events (plus runs of unsynchronised ones)
-		for (; EndDesc->Serial == ESerial::Ignored; ++EndDesc);
-		if (EndDesc->Serial != ESerial::Terminal)
+		if (EndDesc->Serial == NextSerial)
 		{
-			uint32 NextSerial = EndDesc->Serial;
 			do
 			{
+				NextSerial = (NextSerial + 1) & (ESerial::Range - 1);
+
 				do
 				{
 					++EndDesc;
 				}
 				while (EndDesc->Serial == ESerial::Ignored);
-
-				++NextSerial;
 			}
 			while (EndDesc->Serial == NextSerial);
 		}
+		else
+		{
+			// The lowest known serial number is not low enough so we are unable
+			// to proceed any further.
+			break;
+		}
 
 		// Update the heap
 		if (EndDesc->Serial != ESerial::Terminal)
@@ -2844,6 +2855,18 @@ FProtocol5Stage::EStatus FProtocol5Stage::OnDataNormal(const FMachineContext& Co
 	}
 	while (!EventDescHeap.IsEmpty());
 
+	// If there are any streams left in the heap then we are unable to proceed
+	// until more data is received. We'll rewind the streams until more data is
+	// available. It is an efficient way to do things, but it is simple way.
+	for (FEventDescStream& Stream : EventDescHeap)
+	{
+		const FEventDesc& EventDesc = Stream.EventDescs[0];
+		uint32 HeaderSize = 1 + EventDesc.bTwoByteUid + (ESerial::Bits / 8);
+
+		FStreamReader* Reader = Transport.GetThreadStream(Stream.TransportIndex);
+		Reader->Backtrack(EventDesc.Data - HeaderSize);
+	}
+
 	return bNotEnoughData ? EStatus::NotEnoughData : EStatus::Eof;
 }
 
