Renamed Index member to something a little more descriptive

#rnx

diff --git a/Engine/Source/Developer/TraceAnalysis/Private/Analysis/Engine.cpp b/Engine/Source/Developer/TraceAnalysis/Private/Analysis/Engine.cpp
--- a/Engine/Source/Developer/TraceAnalysis/Private/Analysis/Engine.cpp
+++ b/Engine/Source/Developer/TraceAnalysis/Private/Analysis/Engine.cpp
@@ -2707,7 +2707,7 @@ FProtocol5Stage::EStatus FProtocol5Stage::OnDataNormal(const FMachineContext& Co
 		uint32					TransportIndex;
 		union
 		{
-			uint32				Index;
+			uint32				ContainerIndex;
 			const FEventDesc*	EventDescs;
 		};
 	};
@@ -2738,7 +2738,7 @@ FProtocol5Stage::EStatus FProtocol5Stage::OnDataNormal(const FMachineContext& Co
 			FEventDescStream Out;
 			Out.ThreadId = Transport.GetThreadId(i);
 			Out.TransportIndex = i;
-			Out.Index = NumEventDescs;
+			Out.ContainerIndex = NumEventDescs;
 			EventDescHeap.Add(Out);
 		}
 
@@ -2748,7 +2748,7 @@ FProtocol5Stage::EStatus FProtocol5Stage::OnDataNormal(const FMachineContext& Co
 	// Now EventDescs is stable we can convert the indices into pointers
 	for (FEventDescStream& Stream : EventDescHeap)
 	{
-		Stream.EventDescs = EventDescs.GetData() + Stream.Index;
+		Stream.EventDescs = EventDescs.GetData() + Stream.ContainerIndex;
 	}
 
 	// Process leading unsynchronised events so that each stream starts with a
