int -> int32

#rnx

diff --git a/Engine/Source/Developer/TraceAnalysis/Private/Analysis/Engine.cpp b/Engine/Source/Developer/TraceAnalysis/Private/Analysis/Engine.cpp
--- a/Engine/Source/Developer/TraceAnalysis/Private/Analysis/Engine.cpp
+++ b/Engine/Source/Developer/TraceAnalysis/Private/Analysis/Engine.cpp
@@ -2515,7 +2515,7 @@ private:
 	int32					ParseEvents(FStreamReader& Reader, EventDescArray& OutEventDescs);
 	int32					ParseEventsWithAux(FStreamReader& Reader, EventDescArray& OutEventDescs);
 	int32					ParseEvent(FStreamReader& Reader, FEventDesc& OutEventDesc);
-	int						DispatchEvents(FAnalysisBridge& Bridge, const FEventDesc* EventDesc, uint32 Count);
+	int32					DispatchEvents(FAnalysisBridge& Bridge, const FEventDesc* EventDesc, uint32 Count);
 	FTypeRegistry			TypeRegistry;
 	FTidPacketTransport&	Transport;
 	EventDescArray			EventDescs;
@@ -3088,7 +3088,7 @@ int32 FProtocol5Stage::ParseEvent(FStreamReader& Reader, FEventDesc& EventDesc)
 }
 
 ////////////////////////////////////////////////////////////////////////////////
-int FProtocol5Stage::DispatchEvents(
+int32 FProtocol5Stage::DispatchEvents(
 	FAnalysisBridge& Bridge,
 	const FEventDesc* EventDesc, uint32 Count)
 {
