Changed scope of a struct function to class.

#rnx

diff --git a/Engine/Source/Developer/TraceAnalysis/Private/Analysis/Engine.cpp b/Engine/Source/Developer/TraceAnalysis/Private/Analysis/Engine.cpp
--- a/Engine/Source/Developer/TraceAnalysis/Private/Analysis/Engine.cpp
+++ b/Engine/Source/Developer/TraceAnalysis/Private/Analysis/Engine.cpp
@@ -2494,6 +2494,18 @@ private:
 	};
 	static_assert(sizeof(FEventDesc) == 16, "");
 
+	struct alignas(16) FEventDescStream
+	{
+		uint32					ThreadId;
+		uint32					TransportIndex;
+		union
+		{
+			uint32				ContainerIndex;
+			const FEventDesc*	EventDescs;
+		};
+	};
+	static_assert(sizeof(FEventDescStream) == 16, "");
+
 	enum ESerial : int32 
 	{
 		Bits		= 24,
@@ -2701,16 +2713,6 @@ FProtocol5Stage::EStatus FProtocol5Stage::OnDataNormal(const FMachineContext& Co
 	EventDescs.Reset();
 	bool bNotEnoughData = false;
 
-	struct alignas(16) FEventDescStream
-	{
-		uint32					ThreadId;
-		uint32					TransportIndex;
-		union
-		{
-			uint32				ContainerIndex;
-			const FEventDesc*	EventDescs;
-		};
-	};
 	TArray<FEventDescStream> EventDescHeap;
 	EventDescHeap.Reserve(Transport.GetThreadCount());
 
