Use a less-than operator instead of a lambda for the event-desc heap

#rnx

diff --git a/Engine/Source/Developer/TraceAnalysis/Private/Analysis/Engine.cpp b/Engine/Source/Developer/TraceAnalysis/Private/Analysis/Engine.cpp
--- a/Engine/Source/Developer/TraceAnalysis/Private/Analysis/Engine.cpp
+++ b/Engine/Source/Developer/TraceAnalysis/Private/Analysis/Engine.cpp
@@ -2503,6 +2503,8 @@ private:
 			uint32				ContainerIndex;
 			const FEventDesc*	EventDescs;
 		};
+
+		bool operator < (const FEventDescStream& Rhs) const;
 	};
 	static_assert(sizeof(FEventDescStream) == 16, "");
 
@@ -2534,6 +2536,14 @@ private:
 	uint32					NextSerial = ~0u;
 };
 
+////////////////////////////////////////////////////////////////////////////////
+bool FProtocol5Stage::FEventDescStream::operator < (const FEventDescStream& Rhs) const
+{
+	int32 Delta = Rhs.EventDescs->Serial - EventDescs->Serial;
+	int32 Wrapped = uint32(Delta + ESerial::HalfRange - 1) >= uint32(ESerial::Range - 2);
+	return (Wrapped ^ (Delta > 0)) != 0;
+}
+
 ////////////////////////////////////////////////////////////////////////////////
 FProtocol5Stage::FProtocol5Stage(FTransport* InTransport)
 : Transport(*(FTidPacketTransport*)InTransport)
@@ -2795,13 +2805,7 @@ FProtocol5Stage::EStatus FProtocol5Stage::OnDataNormal(const FMachineContext& Co
 	// than half the serial space, they have wrapped.
 
 	// A min-heap is used to peel off groups of events by lowest serial
-	auto Comparison = [] (const FEventDescStream& Lhs, const FEventDescStream& Rhs)
-	{
-		int32 Delta = Rhs.EventDescs->Serial - Lhs.EventDescs->Serial;
-		int32 Wrapped = uint32(Delta + ESerial::HalfRange - 1) >= uint32(ESerial::Range - 2);
-		return (Wrapped ^ (Delta > 0)) != 0;
-	};
-	EventDescHeap.Heapify(Comparison);
+	EventDescHeap.Heapify();
 
 	// Events must be consumed contiguously.
 	if (NextSerial == ~0u)
@@ -2845,7 +2849,7 @@ FProtocol5Stage::EStatus FProtocol5Stage::OnDataNormal(const FMachineContext& Co
 			auto& Out = EventDescHeap.Add_GetRef({Stream.ThreadId, Stream.TransportIndex});
 			Out.EventDescs = EndDesc;
 		}
-		EventDescHeap.HeapPopDiscard(Comparison);
+		EventDescHeap.HeapPopDiscard();
 
 		// Dispatch.
 		int32 DescNum = int32(UPTRINT(EndDesc - StartDesc));
