Split OnNormalEvents() into two methods

#rnx

diff --git a/Engine/Source/Developer/TraceAnalysis/Private/Analysis/Engine.cpp b/Engine/Source/Developer/TraceAnalysis/Private/Analysis/Engine.cpp
--- a/Engine/Source/Developer/TraceAnalysis/Private/Analysis/Engine.cpp
+++ b/Engine/Source/Developer/TraceAnalysis/Private/Analysis/Engine.cpp
@@ -2530,6 +2530,7 @@ private:
 	int32					ParseEventsWithAux(FStreamReader& Reader, EventDescArray& OutEventDescs);
 	int32					ParseEvent(FStreamReader& Reader, FEventDesc& OutEventDesc);
 	int32					DispatchEvents(FAnalysisBridge& Bridge, const FEventDesc* EventDesc, uint32 Count);
+	int32					DispatchEvents(FAnalysisBridge& Bridge, TArray<FEventDescStream>& EventDescHeap);
 	FTypeRegistry			TypeRegistry;
 	FTidPacketTransport&	Transport;
 	EventDescArray			EventDescs;
@@ -2813,13 +2814,26 @@ FProtocol5Stage::EStatus FProtocol5Stage::OnDataNormal(const FMachineContext& Co
 		NextSerial = EventDescHeap.HeapTop().EventDescs[0].Serial;
 	}
 
+	if (DispatchEvents(Context.Bridge, EventDescHeap) < 0)
+	{
+		return EStatus::Error;
+	}
+
+	return bNotEnoughData ? EStatus::NotEnoughData : EStatus::Eof;
+}
+
+////////////////////////////////////////////////////////////////////////////////
+int32 FProtocol5Stage::DispatchEvents(
+	FAnalysisBridge& Bridge,
+	TArray<FEventDescStream>& EventDescHeap)
+{
 	do
 	{
 		const FEventDescStream& Stream = EventDescHeap.HeapTop();
 		const FEventDesc* StartDesc = Stream.EventDescs;
 		const FEventDesc* EndDesc = StartDesc;
 
-		Context.Bridge.SetActiveThread(Stream.ThreadId);
+		Bridge.SetActiveThread(Stream.ThreadId);
 
 		// Extract a run of consecutive events (plus runs of unsynchronised ones)
 		if (EndDesc->Serial == NextSerial)
@@ -2854,9 +2868,9 @@ FProtocol5Stage::EStatus FProtocol5Stage::OnDataNormal(const FMachineContext& Co
 		// Dispatch.
 		int32 DescNum = int32(UPTRINT(EndDesc - StartDesc));
 		check(DescNum > 0);
-		if (DispatchEvents(Context.Bridge, StartDesc, DescNum) < 0)
+		if (DispatchEvents(Bridge, StartDesc, DescNum) < 0)
 		{
-			return EStatus::Error;
+			return -1;
 		}
 	}
 	while (!EventDescHeap.IsEmpty());
@@ -2873,7 +2887,7 @@ FProtocol5Stage::EStatus FProtocol5Stage::OnDataNormal(const FMachineContext& Co
 		Reader->Backtrack(EventDesc.Data - HeaderSize);
 	}
 
-	return bNotEnoughData ? EStatus::NotEnoughData : EStatus::Eof;
+	return 0;
 }
 
 ////////////////////////////////////////////////////////////////////////////////
