Always process both classes of trace events each update

#rnx

diff --git a/Engine/Source/Developer/TraceAnalysis/Private/Analysis/Engine.cpp b/Engine/Source/Developer/TraceAnalysis/Private/Analysis/Engine.cpp
--- a/Engine/Source/Developer/TraceAnalysis/Private/Analysis/Engine.cpp
+++ b/Engine/Source/Developer/TraceAnalysis/Private/Analysis/Engine.cpp
@@ -2560,13 +2560,23 @@ FProtocol5Stage::EStatus FProtocol5Stage::OnData(
 	Transport.SetReader(Reader);
 	Transport.Update();
 
+	bool bNotEnoughData;
+
 	EStatus Ret = OnDataImportant(Context);
-	if (Ret != EStatus::Eof)
+	if (Ret == EStatus::Error)
+	{
+		return Ret;
+	}
+	bNotEnoughData = (Ret != EStatus::Eof);
+
+	Ret = OnDataNormal(Context);
+	if (Ret == EStatus::Error)
 	{
 		return Ret;
 	}
+	bNotEnoughData |= (Ret == EStatus::NotEnoughData);
 
-	return OnDataNormal(Context);
+	return bNotEnoughData ? EStatus::NotEnoughData : EStatus::Eof;
 }
 
 ////////////////////////////////////////////////////////////////////////////////
