Moved the heap update out of the loop and into a lambda.

The update is also done after dispatching events to analyzers now instead of before so the Stream object is valid for the entire scope

#rnx

diff --git a/Engine/Source/Developer/TraceAnalysis/Private/Analysis/Engine.cpp b/Engine/Source/Developer/TraceAnalysis/Private/Analysis/Engine.cpp
--- a/Engine/Source/Developer/TraceAnalysis/Private/Analysis/Engine.cpp
+++ b/Engine/Source/Developer/TraceAnalysis/Private/Analysis/Engine.cpp
@@ -2837,13 +2837,24 @@ int32 FProtocol5Stage::DispatchEvents(
 	FAnalysisBridge& Bridge,
 	TArray<FEventDescStream>& EventDescHeap)
 {
+	auto UpdateHeap = [&] (const FEventDescStream& Stream, const FEventDesc* EventDesc)
+	{
+		if (EventDesc->Serial != ESerial::Terminal)
+		{
+			FEventDescStream Next = Stream;
+			Next.EventDescs = EventDesc;
+			EventDescHeap.Add(Next);
+		}
+
+		EventDescHeap.HeapPopDiscard();
+	};
+
 	do
 	{
 		const FEventDescStream& Stream = EventDescHeap.HeapTop();
 		const FEventDesc* StartDesc = Stream.EventDescs;
 		const FEventDesc* EndDesc = StartDesc;
 
-		Bridge.SetActiveThread(Stream.ThreadId);
 
 		// Extract a run of consecutive events (plus runs of unsynchronised ones)
 		if (EndDesc->Serial == NextSerial)
@@ -2867,21 +2878,16 @@ int32 FProtocol5Stage::DispatchEvents(
 			break;
 		}
 
-		// Update the heap
-		if (EndDesc->Serial != ESerial::Terminal)
-		{
-			auto& Out = EventDescHeap.Add_GetRef({Stream.ThreadId, Stream.TransportIndex});
-			Out.EventDescs = EndDesc;
-		}
-		EventDescHeap.HeapPopDiscard();
-
 		// Dispatch.
+		Bridge.SetActiveThread(Stream.ThreadId);
 		int32 DescNum = int32(UPTRINT(EndDesc - StartDesc));
 		check(DescNum > 0);
 		if (DispatchEvents(Bridge, StartDesc, DescNum) < 0)
 		{
 			return -1;
 		}
+
+		UpdateHeap(Stream, EndDesc);
 	}
 	while (!EventDescHeap.IsEmpty());
 
