Use 'Mask' instead of 'Range - 1'

#rnx

diff --git a/Engine/Source/Developer/TraceAnalysis/Private/Analysis/Engine.cpp b/Engine/Source/Developer/TraceAnalysis/Private/Analysis/Engine.cpp
--- a/Engine/Source/Developer/TraceAnalysis/Private/Analysis/Engine.cpp
+++ b/Engine/Source/Developer/TraceAnalysis/Private/Analysis/Engine.cpp
@@ -2861,7 +2861,7 @@ int32 FProtocol5Stage::DispatchEvents(
 		{
 			do
 			{
-				NextSerial = (NextSerial + 1) & (ESerial::Range - 1);
+				NextSerial = (NextSerial + 1) & ESerial::Mask;
 
 				do
 				{
