Use sync points embedded in the transport to reliably detect gaps in trace event serial numbers that can occur when late connecting.

#rnx

diff --git a/Engine/Source/Developer/TraceAnalysis/Private/Analysis/Engine.cpp b/Engine/Source/Developer/TraceAnalysis/Private/Analysis/Engine.cpp
--- a/Engine/Source/Developer/TraceAnalysis/Private/Analysis/Engine.cpp
+++ b/Engine/Source/Developer/TraceAnalysis/Private/Analysis/Engine.cpp
@@ -2490,7 +2490,11 @@ private:
 			};
 			uint64			Meta = 0;
 		};
-		const uint8*		Data = nullptr;
+		union
+		{
+			uint32			GapLength;
+			const uint8*	Data;
+		};
 	};
 	static_assert(sizeof(FEventDesc) == 16, "");
 
@@ -2504,6 +2508,8 @@ private:
 			const FEventDesc*	EventDescs;
 		};
 
+		enum { GapThreadId = ~0u };
+
 		bool operator < (const FEventDescStream& Rhs) const;
 	};
 	static_assert(sizeof(FEventDescStream) == 16, "");
@@ -2531,10 +2537,15 @@ private:
 	int32					ParseEvent(FStreamReader& Reader, FEventDesc& OutEventDesc);
 	int32					DispatchEvents(FAnalysisBridge& Bridge, const FEventDesc* EventDesc, uint32 Count);
 	int32					DispatchEvents(FAnalysisBridge& Bridge, TArray<FEventDescStream>& EventDescHeap);
+	void					DetectSerialGaps(TArray<FEventDescStream>& EventDescHeap);
+	template <typename CALLBACK>
+	void					ForEachSerialGap(const TArray<FEventDescStream>& EventDescHeap, CALLBACK&& Callback);
 	FTypeRegistry			TypeRegistry;
 	FTidPacketTransport&	Transport;
 	EventDescArray			EventDescs;
+	EventDescArray			SerialGaps;
 	uint32					NextSerial = ~0u;
+	uint32					SyncCount;
 };
 
 ////////////////////////////////////////////////////////////////////////////////
@@ -2548,6 +2559,7 @@ bool FProtocol5Stage::FEventDescStream::operator < (const FEventDescStream& Rhs)
 ////////////////////////////////////////////////////////////////////////////////
 FProtocol5Stage::FProtocol5Stage(FTransport* InTransport)
 : Transport(*(FTidPacketTransport*)InTransport)
+, SyncCount(Transport.GetSyncCount())
 {
 	EventDescs.Reserve(8 << 10);
 }
@@ -2819,11 +2831,13 @@ FProtocol5Stage::EStatus FProtocol5Stage::OnDataNormal(const FMachineContext& Co
 	EventDescHeap.Heapify();
 
 	// Events must be consumed contiguously.
-	if (NextSerial == ~0u)
+	if (NextSerial == ~0u && Transport.GetSyncCount())
 	{
 		NextSerial = EventDescHeap.HeapTop().EventDescs[0].Serial;
 	}
 
+	DetectSerialGaps(EventDescHeap);
+
 	if (DispatchEvents(Context.Bridge, EventDescHeap) < 0)
 	{
 		return EStatus::Error;
@@ -2855,6 +2869,16 @@ int32 FProtocol5Stage::DispatchEvents(
 		const FEventDesc* StartDesc = Stream.EventDescs;
 		const FEventDesc* EndDesc = StartDesc;
 
+		// DetectSerialGaps() will add a special stream that communicates gaps
+		// in in serial numbers, gaps that will never be resolved. Thread IDs
+		// are uint16 everywhere else so they will never collide with GapThreadId.
+		if (Stream.ThreadId == FEventDescStream::GapThreadId)
+		{
+			NextSerial = EndDesc->Serial + EndDesc->GapLength;
+			NextSerial &= ESerial::Mask;
+			UpdateHeap(Stream, EndDesc + 1);
+			continue;
+		}
 
 		// Extract a run of consecutive events (plus runs of unsynchronised ones)
 		if (EndDesc->Serial == NextSerial)
@@ -3123,6 +3147,150 @@ int32 FProtocol5Stage::ParseEvent(FStreamReader& Reader, FEventDesc& EventDesc)
 	return HeaderSize + EventSize;
 }
 
+////////////////////////////////////////////////////////////////////////////////
+template <typename CALLBACK>
+void FProtocol5Stage::ForEachSerialGap(
+	const TArray<FEventDescStream>& EventDescHeap,
+	CALLBACK&& Callback)
+{
+	TArray<FEventDescStream> HeapCopy(EventDescHeap);
+	int Serial = HeapCopy.HeapTop().EventDescs[0].Serial;
+
+	// There might be a gap at the beginning of the heap if some events have
+	// already been consumed.
+	if (NextSerial != Serial)
+	{
+		if (!Callback(NextSerial, Serial))
+		{
+			return;
+		}
+	}
+
+	// A min-heap is used to peel off each stream (thread) with the lowest serial
+	// numbered event.
+	do
+	{
+		const FEventDescStream& Stream = HeapCopy.HeapTop();
+		const FEventDesc* EventDesc = Stream.EventDescs;
+
+		// If the next lowest serial number doesn't match where we got up to in
+		// the previous stream we have found a gap. Celebration ensues.
+		if (Serial != EventDesc->Serial)
+		{
+			if (!Callback(Serial, EventDesc->Serial))
+			{
+				return;
+			}
+		}
+
+		// Consume consecutive events (including unsynchronised ones).
+		Serial = EventDesc->Serial;
+		do
+		{
+			do
+			{
+				++EventDesc;
+			}
+			while (EventDesc->Serial == ESerial::Ignored);
+
+			Serial = (Serial + 1) & ESerial::Mask;
+		}
+		while (EventDesc->Serial == Serial);
+
+		// Update the heap
+		if (EventDesc->Serial != ESerial::Terminal)
+		{
+			auto& Out = HeapCopy.Add_GetRef({Stream.ThreadId, Stream.TransportIndex});
+			Out.EventDescs = EventDesc;
+		}
+		HeapCopy.HeapPopDiscard();
+	}
+	while (!HeapCopy.IsEmpty());
+}
+
+////////////////////////////////////////////////////////////////////////////////
+void FProtocol5Stage::DetectSerialGaps(TArray<FEventDescStream>& EventDescHeap)
+{
+	// Events that should be synchronised across threads are assigned serial
+	// numbers so they can be analysed in the correct order. Gaps in the
+	// serials can occur under two scenarios; 1) when packets are dropped from
+	// the trace tail to make space for new trace events, and 2) when Trace's
+	// worker thread ticks, samples all the trace buffers and sends their data.
+	// In late-connect scenarios these gaps need to be skipped over in order to
+	// successfully reserialise events in the data stream. To further complicate
+	// matters, most of the gaps from (2) will get filled by the following update,
+	// leading to initial false positive gaps. By embedding sync points in the
+	// stream we can reliably differentiate genuine gaps from temporary ones.
+	//
+	// Note that this could be done without sync points but it is an altogether
+	// more complex solution. So unsightly embedded syncs it is...
+
+	if (SyncCount == Transport.GetSyncCount())
+	{
+		return;
+	}
+
+	SyncCount = Transport.GetSyncCount();
+
+	if (SyncCount == 1)
+	{
+		// On the first update we will just collect gaps.
+		auto GatherGap = [this] (int32 Lhs, int32 Rhs)
+		{
+			FEventDesc& Gap = SerialGaps.Emplace_GetRef();
+			Gap.Serial = Lhs;
+			Gap.GapLength = (Rhs - Lhs) & ESerial::Mask;
+			return true;
+		};
+		ForEachSerialGap(EventDescHeap, GatherGap);
+	}
+	else
+	{
+		// On the second update we detect where gaps from the previous update
+		// start getting filled in. Any gaps preceding that point are genuine.
+		uint32 GapCount = 0;
+		auto RecordGap = [this, &GapCount] (int32 Lhs, int32 Rhs)
+		{
+			for (uint32 n = SerialGaps.Num(); GapCount < n; ++GapCount)
+			{
+				FEventDesc& SerialGap = SerialGaps[GapCount];
+
+				if (SerialGap.Serial >= Lhs)
+				{
+					GapCount += (SerialGap.Serial == Lhs);
+					break;
+				}
+
+				return false;
+			}
+			return true;
+		};
+
+		ForEachSerialGap(EventDescHeap, RecordGap);
+
+		if (GapCount == 0)
+		{
+			SerialGaps.Empty();
+			return;
+		}
+
+		// Turn the genuine gaps into a stream that DispatchEvents() can handle
+		// and use to skip over them.
+
+		if (GapCount == uint32(SerialGaps.Num()))
+		{
+			SerialGaps.Emplace();
+		}
+		FEventDesc& Terminator = SerialGaps[GapCount];
+		Terminator.Serial = ESerial::Terminal;
+
+		FEventDescStream Out = EventDescHeap[0];
+		Out.ThreadId = FEventDescStream::GapThreadId;
+		Out.EventDescs = SerialGaps.GetData();
+		EventDescHeap.HeapPush(Out);
+	}
+}
+
 ////////////////////////////////////////////////////////////////////////////////
 int32 FProtocol5Stage::DispatchEvents(
 	FAnalysisBridge& Bridge,
