Analyze new-event events before anything else.

Previously new-event and important events were collected in a pass together. However, because the new-events hadn't been processed yet, the important event stream was unintelligable until the next update. This would accellerate normal events past important ones that they might depend upon.

#rnx

diff --git a/Engine/Source/Developer/TraceAnalysis/Private/Analysis/Engine.cpp b/Engine/Source/Developer/TraceAnalysis/Private/Analysis/Engine.cpp
--- a/Engine/Source/Developer/TraceAnalysis/Private/Analysis/Engine.cpp
+++ b/Engine/Source/Developer/TraceAnalysis/Private/Analysis/Engine.cpp
@@ -2529,6 +2529,7 @@ private:
 
 
 	virtual EStatus			OnData(FStreamReader& Reader, const FMachineContext& Context) override;
+	EStatus					OnDataNewEvents(const FMachineContext& Context);
 	EStatus					OnDataImportant(const FMachineContext& Context);
 	EStatus					OnDataNormal(const FMachineContext& Context);
 	int32					ParseImportantEvents(FStreamReader& Reader, EventDescArray& OutEventDescs);
@@ -2572,15 +2573,25 @@ FProtocol5Stage::EStatus FProtocol5Stage::OnData(
 	Transport.SetReader(Reader);
 	Transport.Update();
 
+	// New-events. They must be processed before anything else otherwise events
+	// can not be interpreted.
+	EStatus Ret = OnDataNewEvents(Context);
+	if (Ret == EStatus::Error)
+	{
+		return Ret;
+	}
+
 	bool bNotEnoughData;
 
-	EStatus Ret = OnDataImportant(Context);
+	// Important events
+	Ret = OnDataImportant(Context);
 	if (Ret == EStatus::Error)
 	{
 		return Ret;
 	}
 	bNotEnoughData = (Ret != EStatus::Eof);
 
+	// Normal events
 	Ret = OnDataNormal(Context);
 	if (Ret == EStatus::Error)
 	{
@@ -2592,50 +2603,50 @@ FProtocol5Stage::EStatus FProtocol5Stage::OnData(
 }
 
 ////////////////////////////////////////////////////////////////////////////////
-FProtocol5Stage::EStatus FProtocol5Stage::OnDataImportant(const FMachineContext& Context)
+FProtocol5Stage::EStatus FProtocol5Stage::OnDataNewEvents(const FMachineContext& Context)
 {
-	// New and important event packets
-
 	EventDescs.Reset();
-	bool bNotEnoughData = false;
 
-	for (uint32 i = 0, n = ETransportTid::Bias; i < n; ++i)
+	FStreamReader* ThreadReader = Transport.GetThreadStream(ETransportTid::Events);
+	if (ThreadReader->IsEmpty())
 	{
-		FStreamReader* ThreadReader = Transport.GetThreadStream(i);
-		if (ThreadReader->IsEmpty())
-		{
-			continue;
-		}
-
-		if (ParseImportantEvents(*ThreadReader, EventDescs) < 0)
-		{
-			return EStatus::Error;
-		}
-
-		bNotEnoughData |= !ThreadReader->IsEmpty();
+		return EStatus::Eof;
 	}
 
-	// Dispatch and filter out new-event events 
-	for (uint32 j = 0, k = 0, m = EventDescs.Num(); ; ++j)
+	if (ParseImportantEvents(*ThreadReader, EventDescs) < 0)
 	{
-		if (j >= m)
-		{
-			EventDescs.SetNum(k);
-			break;
-		}
-
-		const FEventDesc& EventDesc = EventDescs[j];
-		if (EventDesc.Uid != EKnownUids::NewEvent)
-		{
-			EventDescs[k] = EventDesc;
-			++k;
-			continue;
-		}
+		return EStatus::Error;
+	}
 
+	for (const FEventDesc& EventDesc : EventDescs)
+	{
 		const FTypeRegistry::FTypeInfo* TypeInfo = TypeRegistry.Add(EventDesc.Data);
 		Context.Bridge.OnNewType(TypeInfo);
 	}
 
+	return EStatus::Eof;
+}
+
+////////////////////////////////////////////////////////////////////////////////
+FProtocol5Stage::EStatus FProtocol5Stage::OnDataImportant(const FMachineContext& Context)
+{
+	static_assert(ETransportTid::Importants == ETransportTid::Internal, "It is assumed there is only one 'important' thread stream");
+
+	EventDescs.Reset();
+
+	FStreamReader* ThreadReader = Transport.GetThreadStream(ETransportTid::Importants);
+	if (ThreadReader->IsEmpty())
+	{
+		return EStatus::Eof;
+	}
+
+	if (ParseImportantEvents(*ThreadReader, EventDescs) < 0)
+	{
+		return EStatus::Error;
+	}
+
+	bool bNotEnoughData = !ThreadReader->IsEmpty();
+
 	if (EventDescs.Num() <= 0)
 	{
 		return bNotEnoughData ? EStatus::NotEnoughData : EStatus::Eof;
