EStatus::Eof -> EStatus::EndOfStream

#rnx

diff --git a/Engine/Source/Developer/TraceAnalysis/Private/Analysis/Engine.cpp b/Engine/Source/Developer/TraceAnalysis/Private/Analysis/Engine.cpp
--- a/Engine/Source/Developer/TraceAnalysis/Private/Analysis/Engine.cpp
+++ b/Engine/Source/Developer/TraceAnalysis/Private/Analysis/Engine.cpp
@@ -1628,7 +1628,7 @@ public:
 		Error,
 		Abort,
 		NotEnoughData,
-		Eof,
+		EndOfStream,
 		Continue,
 	};
 
@@ -1820,7 +1820,7 @@ FProtocol0Stage::EStatus FProtocol0Stage::OnData(
 		Transport->Advance(BlockSize);
 	}
 
-	return Reader.IsEmpty() ? EStatus::Eof : EStatus::NotEnoughData;
+	return Reader.IsEmpty() ? EStatus::EndOfStream : EStatus::NotEnoughData;
 }
 
 
@@ -2065,7 +2065,7 @@ FProtocol2Stage::EStatus FProtocol2Stage::OnData(
 		}
 	}
 
-	return Reader.IsEmpty() ? EStatus::Eof : EStatus::NotEnoughData;
+	return Reader.IsEmpty() ? EStatus::EndOfStream : EStatus::NotEnoughData;
 }
 
 ////////////////////////////////////////////////////////////////////////////////
@@ -2589,7 +2589,7 @@ FProtocol5Stage::EStatus FProtocol5Stage::OnData(
 	{
 		return Ret;
 	}
-	bNotEnoughData = (Ret != EStatus::Eof);
+	bNotEnoughData = (Ret != EStatus::EndOfStream);
 
 	// Normal events
 	Ret = OnDataNormal(Context);
@@ -2599,7 +2599,7 @@ FProtocol5Stage::EStatus FProtocol5Stage::OnData(
 	}
 	bNotEnoughData |= (Ret == EStatus::NotEnoughData);
 
-	return bNotEnoughData ? EStatus::NotEnoughData : EStatus::Eof;
+	return bNotEnoughData ? EStatus::NotEnoughData : EStatus::EndOfStream;
 }
 
 ////////////////////////////////////////////////////////////////////////////////
@@ -2610,7 +2610,7 @@ FProtocol5Stage::EStatus FProtocol5Stage::OnDataNewEvents(const FMachineContext&
 	FStreamReader* ThreadReader = Transport.GetThreadStream(ETransportTid::Events);
 	if (ThreadReader->IsEmpty())
 	{
-		return EStatus::Eof;
+		return EStatus::EndOfStream;
 	}
 
 	if (ParseImportantEvents(*ThreadReader, EventDescs) < 0)
@@ -2624,7 +2624,7 @@ FProtocol5Stage::EStatus FProtocol5Stage::OnDataNewEvents(const FMachineContext&
 		Context.Bridge.OnNewType(TypeInfo);
 	}
 
-	return EStatus::Eof;
+	return EStatus::EndOfStream;
 }
 
 ////////////////////////////////////////////////////////////////////////////////
@@ -2637,7 +2637,7 @@ FProtocol5Stage::EStatus FProtocol5Stage::OnDataImportant(const FMachineContext&
 	FStreamReader* ThreadReader = Transport.GetThreadStream(ETransportTid::Importants);
 	if (ThreadReader->IsEmpty())
 	{
-		return EStatus::Eof;
+		return EStatus::EndOfStream;
 	}
 
 	if (ParseImportantEvents(*ThreadReader, EventDescs) < 0)
@@ -2649,7 +2649,7 @@ FProtocol5Stage::EStatus FProtocol5Stage::OnDataImportant(const FMachineContext&
 
 	if (EventDescs.Num() <= 0)
 	{
-		return bNotEnoughData ? EStatus::NotEnoughData : EStatus::Eof;
+		return bNotEnoughData ? EStatus::NotEnoughData : EStatus::EndOfStream;
 	}
 
 	// Dispatch looks ahead to the next desc looking for runs of aux blobs. As
@@ -2663,7 +2663,7 @@ FProtocol5Stage::EStatus FProtocol5Stage::OnDataImportant(const FMachineContext&
 		return EStatus::Error;
 	}
 
-	return bNotEnoughData ? EStatus::NotEnoughData : EStatus::Eof;
+	return bNotEnoughData ? EStatus::NotEnoughData : EStatus::EndOfStream;
 }
 
 ////////////////////////////////////////////////////////////////////////////////
@@ -2830,7 +2830,7 @@ FProtocol5Stage::EStatus FProtocol5Stage::OnDataNormal(const FMachineContext& Co
 	// Early out if there isn't any events available.
 	if (UNLIKELY(EventDescHeap.IsEmpty()))
 	{
-		return bNotEnoughData ? EStatus::NotEnoughData : EStatus::Eof;
+		return bNotEnoughData ? EStatus::NotEnoughData : EStatus::EndOfStream;
 	}
 
 	// Provided that less than approximately "SerialRange * BytesPerSerial"
@@ -2854,7 +2854,7 @@ FProtocol5Stage::EStatus FProtocol5Stage::OnDataNormal(const FMachineContext& Co
 		return EStatus::Error;
 	}
 
-	return bNotEnoughData ? EStatus::NotEnoughData : EStatus::Eof;
+	return bNotEnoughData ? EStatus::NotEnoughData : EStatus::EndOfStream;
 }
 
 ////////////////////////////////////////////////////////////////////////////////
