OnData() returns Continue if there's more work to do.

If there is more data available but none of the OnData() steps have indicated the they require more data, then OnData() can return Continue. This will allow the outer machine to make more progress, and ensure that smaller traces (where a single read can contain all sync points) makes progress.

#rnx

diff --git a/Engine/Source/Developer/TraceAnalysis/Private/Analysis/Engine.cpp b/Engine/Source/Developer/TraceAnalysis/Private/Analysis/Engine.cpp
--- a/Engine/Source/Developer/TraceAnalysis/Private/Analysis/Engine.cpp
+++ b/Engine/Source/Developer/TraceAnalysis/Private/Analysis/Engine.cpp
@@ -2599,7 +2599,12 @@ FProtocol5Stage::EStatus FProtocol5Stage::OnData(
 	}
 	bNotEnoughData |= (Ret == EStatus::NotEnoughData);
 
-	return bNotEnoughData ? EStatus::NotEnoughData : EStatus::EndOfStream;
+	if (bNotEnoughData)
+	{
+		return EStatus::NotEnoughData;
+	}
+
+	return (Reader.GetRemaining() != 0) ? EStatus::Continue : EStatus::EndOfStream;
 }
 
 ////////////////////////////////////////////////////////////////////////////////
