Consider if a buffer is in demand (someone asked for more data than it can provide) as a condition of processing more trace data.

While there may be data remaining in the buffer, there might not be enough to decode another transport packet. Thus it is not correct to only consider the remaining number of bytes in the buffer as a condition of continuing to process data. Fortunately there is a demain hint that is set when someone asks for more data than a stream reader can provide. This can be used to differentiate between data remaining in the buffer, and not being able to do anything more with it.

#rm 5.0

diff --git a/Engine/Source/Developer/TraceAnalysis/Private/Analysis/Engine.cpp b/Engine/Source/Developer/TraceAnalysis/Private/Analysis/Engine.cpp
--- a/Engine/Source/Developer/TraceAnalysis/Private/Analysis/Engine.cpp
+++ b/Engine/Source/Developer/TraceAnalysis/Private/Analysis/Engine.cpp
@@ -2604,7 +2604,7 @@ FProtocol5Stage::EStatus FProtocol5Stage::OnData(
 		return EStatus::NotEnoughData;
 	}
 
-	return (Reader.GetRemaining() != 0) ? EStatus::Continue : EStatus::EndOfStream;
+	return Reader.CanMeetDemand() ? EStatus::Continue : EStatus::EndOfStream;
 }
 
 ////////////////////////////////////////////////////////////////////////////////
diff --git a/Engine/Source/Developer/TraceAnalysis/Private/Analysis/StreamReader.cpp b/Engine/Source/Developer/TraceAnalysis/Private/Analysis/StreamReader.cpp
--- a/Engine/Source/Developer/TraceAnalysis/Private/Analysis/StreamReader.cpp
+++ b/Engine/Source/Developer/TraceAnalysis/Private/Analysis/StreamReader.cpp
@@ -37,6 +37,12 @@ int32 FStreamReader::GetRemaining() const
 	return int32(PTRINT(End - Cursor));
 }
 
+////////////////////////////////////////////////////////////////////////////////
+bool FStreamReader::CanMeetDemand() const
+{
+	return (GetRemaining() >= int32(DemandHint));
+}
+
 ////////////////////////////////////////////////////////////////////////////////
 bool FStreamReader::IsEmpty() const
 {
diff --git a/Engine/Source/Developer/TraceAnalysis/Private/Analysis/StreamReader.h b/Engine/Source/Developer/TraceAnalysis/Private/Analysis/StreamReader.h
--- a/Engine/Source/Developer/TraceAnalysis/Private/Analysis/StreamReader.h
+++ b/Engine/Source/Developer/TraceAnalysis/Private/Analysis/StreamReader.h
@@ -21,6 +21,7 @@ public:
 	void						Advance(uint32 Size);
 	bool						IsEmpty() const;
 	int32						GetRemaining() const;
+	bool						CanMeetDemand() const;
 	bool						Backtrack(const uint8* To);
 	struct FMark*				SaveMark() const;
 	void						RestoreMark(struct FMark* Mark);
