Suppress false-positive warning from PVS.

diff --git a/Engine/Source/Developer/TraceAnalysis/Private/Analysis/Engine.cpp b/Engine/Source/Developer/TraceAnalysis/Private/Analysis/Engine.cpp
--- a/Engine/Source/Developer/TraceAnalysis/Private/Analysis/Engine.cpp
+++ b/Engine/Source/Developer/TraceAnalysis/Private/Analysis/Engine.cpp
@@ -3284,7 +3284,7 @@ void FProtocol5Stage::DetectSerialGaps(TArray<FEventDescStream>& EventDescHeap)
 
 		ForEachSerialGap(EventDescHeap, RecordGap);
 
-		if (GapCount == 0)
+		if (GapCount == 0) //-V547
 		{
 			SerialGaps.Empty();
 			return;
