CALLBACK can be defined elsewhere and cause compile errors in this file. Plus SHOUTY template parameters aren't conformant.

#rnx

diff --git a/Engine/Source/Developer/TraceAnalysis/Private/Analysis/Engine.cpp b/Engine/Source/Developer/TraceAnalysis/Private/Analysis/Engine.cpp
--- a/Engine/Source/Developer/TraceAnalysis/Private/Analysis/Engine.cpp
+++ b/Engine/Source/Developer/TraceAnalysis/Private/Analysis/Engine.cpp
@@ -2539,8 +2539,8 @@ private:
 	int32					DispatchEvents(FAnalysisBridge& Bridge, const FEventDesc* EventDesc, uint32 Count);
 	int32					DispatchEvents(FAnalysisBridge& Bridge, TArray<FEventDescStream>& EventDescHeap);
 	void					DetectSerialGaps(TArray<FEventDescStream>& EventDescHeap);
-	template <typename CALLBACK>
-	void					ForEachSerialGap(const TArray<FEventDescStream>& EventDescHeap, CALLBACK&& Callback);
+	template <typename Callback>
+	void					ForEachSerialGap(const TArray<FEventDescStream>& EventDescHeap, Callback&& InCallback);
 	FTypeRegistry			TypeRegistry;
 	FTidPacketTransport&	Transport;
 	EventDescArray			EventDescs;
@@ -3164,10 +3164,10 @@ int32 FProtocol5Stage::ParseEvent(FStreamReader& Reader, FEventDesc& EventDesc)
 }
 
 ////////////////////////////////////////////////////////////////////////////////
-template <typename CALLBACK>
+template <typename Callback>
 void FProtocol5Stage::ForEachSerialGap(
 	const TArray<FEventDescStream>& EventDescHeap,
-	CALLBACK&& Callback)
+	Callback&& InCallback)
 {
 	TArray<FEventDescStream> HeapCopy(EventDescHeap);
 	int Serial = HeapCopy.HeapTop().EventDescs[0].Serial;
@@ -3176,7 +3176,7 @@ void FProtocol5Stage::ForEachSerialGap(
 	// already been consumed.
 	if (NextSerial != Serial)
 	{
-		if (!Callback(NextSerial, Serial))
+		if (!InCallback(NextSerial, Serial))
 		{
 			return;
 		}
@@ -3193,7 +3193,7 @@ void FProtocol5Stage::ForEachSerialGap(
 		// the previous stream we have found a gap. Celebration ensues.
 		if (Serial != EventDesc->Serial)
 		{
-			if (!Callback(Serial, EventDesc->Serial))
+			if (!InCallback(Serial, EventDesc->Serial))
 			{
 				return;
 			}
