Missing const on a less-than-operator

#rnx

diff --git a/Engine/Source/Developer/TraceAnalysis/Private/Analysis/Engine.cpp b/Engine/Source/Developer/TraceAnalysis/Private/Analysis/Engine.cpp
--- a/Engine/Source/Developer/TraceAnalysis/Private/Analysis/Engine.cpp
+++ b/Engine/Source/Developer/TraceAnalysis/Private/Analysis/Engine.cpp
@@ -1915,7 +1915,7 @@ FProtocol2Stage::EStatus FProtocol2Stage::OnData(
 		uint32				ThreadId;
 		FStreamReader*		Reader;
 
-		bool operator < (const FRotaItem& Rhs)
+		bool operator < (const FRotaItem& Rhs) const
 		{
 			int32 Delta = Rhs.Serial - Serial;
 			int32 Wrapped = uint32(Delta + 0x007f'fffe) >= uint32(0x00ff'fffd);
