- Changed deleter of Buffer to be the class that allocated it
- Implemented move-constructor/assignment operator. Transport keeps an array threads, each of which of owns a FStreamBuffer instance. When Transport resizes this array, it copies items about. Because there are no move constructors, the compiler ends up cloning FStreamBuffers (via copy-constructing) and incorrectly deleting Buffer pointers still in use.
#rnx

diff --git a/Engine/Source/Developer/TraceAnalysis/Private/Analysis/StreamReader.cpp b/Engine/Source/Developer/TraceAnalysis/Private/Analysis/StreamReader.cpp
--- a/Engine/Source/Developer/TraceAnalysis/Private/Analysis/StreamReader.cpp
+++ b/Engine/Source/Developer/TraceAnalysis/Private/Analysis/StreamReader.cpp
@@ -7,12 +7,6 @@
 namespace UE {
 namespace Trace {
 
-////////////////////////////////////////////////////////////////////////////////
-FStreamReader::~FStreamReader()
-{
-	FMemory::Free(Buffer);
-}
-
 ////////////////////////////////////////////////////////////////////////////////
 const uint8* FStreamReader::GetPointer(uint32 Size)
 {
diff --git a/Engine/Source/Developer/TraceAnalysis/Private/Analysis/StreamReader.h b/Engine/Source/Developer/TraceAnalysis/Private/Analysis/StreamReader.h
--- a/Engine/Source/Developer/TraceAnalysis/Private/Analysis/StreamReader.h
+++ b/Engine/Source/Developer/TraceAnalysis/Private/Analysis/StreamReader.h
@@ -11,7 +11,6 @@ namespace Trace {
 class FStreamReader
 {
 public:
-								~FStreamReader();
 	template <typename Type>
 	Type const*					GetPointer();
 	template <typename Type>
@@ -60,6 +59,12 @@ class FStreamBuffer
 	: public FStreamReader
 {
 public:
+								FStreamBuffer() = default;
+								~FStreamBuffer();
+								FStreamBuffer(FStreamBuffer&& Rhs) noexcept;
+								FStreamBuffer(const FStreamBuffer&)	= default;
+	FStreamBuffer&				operator = (FStreamBuffer&& Rhs) noexcept;
+	FStreamBuffer&				operator = (const FStreamBuffer&)	= delete;
 	template <typename Lambda>
 	int32						Fill(Lambda&& Source);
 	void						Append(const uint8* Data, uint32 Size);
@@ -70,6 +75,33 @@ protected:
 	uint32						BufferSize = 0;
 };
 
+////////////////////////////////////////////////////////////////////////////////
+inline FStreamBuffer::~FStreamBuffer()
+{
+	FMemory::Free(Buffer);
+}
+
+////////////////////////////////////////////////////////////////////////////////
+inline FStreamBuffer::FStreamBuffer(FStreamBuffer&& Rhs) noexcept
+{
+	Swap(BufferSize, Rhs.BufferSize);
+	Swap(Buffer, Rhs.Buffer);
+	Swap(Cursor, Rhs.Cursor);
+	Swap(End, Rhs.End);
+	Swap(DemandHint, Rhs.DemandHint);
+}
+
+////////////////////////////////////////////////////////////////////////////////
+inline FStreamBuffer& FStreamBuffer::operator = (FStreamBuffer&& Rhs) noexcept
+{
+	Swap(BufferSize, Rhs.BufferSize);
+	Swap(Buffer, Rhs.Buffer);
+	Swap(Cursor, Rhs.Cursor);
+	Swap(End, Rhs.End);
+	Swap(DemandHint, Rhs.DemandHint);
+	return *this;
+}
+
 ////////////////////////////////////////////////////////////////////////////////
 template <typename Lambda>
 inline int32 FStreamBuffer::Fill(Lambda&& Source)
