Be explicit about the character size we're expecting.

If TCHAR happens to change to something other than a 16bit integer (e.g. because wchar_t is four bytes on a platform, or because the engine changes to UTF8), the analysis engine will no longer be able to read traces with char16_t-type strings.

#rnx

diff --git a/Engine/Source/Developer/TraceAnalysis/Private/Analysis/Engine.cpp b/Engine/Source/Developer/TraceAnalysis/Private/Analysis/Engine.cpp
--- a/Engine/Source/Developer/TraceAnalysis/Private/Analysis/Engine.cpp
+++ b/Engine/Source/Developer/TraceAnalysis/Private/Analysis/Engine.cpp
@@ -689,7 +689,7 @@ bool IAnalyzer::FEventData::GetString(const ANSICHAR* FieldName, FAnsiStringView
 }
 
 ////////////////////////////////////////////////////////////////////////////////
-bool IAnalyzer::FEventData::GetString(const ANSICHAR* FieldName, FStringView& Out) const
+bool IAnalyzer::FEventData::GetString(const ANSICHAR* FieldName, FWideStringView& Out) const
 {
 	const auto* Info = (const FEventDataInfo*)this;
 	const auto& Dispatch = Info->Dispatch;
@@ -701,18 +701,18 @@ bool IAnalyzer::FEventData::GetString(const ANSICHAR* FieldName, FStringView& Ou
 	}
 
 	const auto& Field = Dispatch.Fields[Index];
-	if (Field.Class != UE::Trace::Protocol0::Field_String || Field.SizeAndType != sizeof(TCHAR))
+	if (Field.Class != UE::Trace::Protocol0::Field_String || Field.SizeAndType != sizeof(WIDECHAR))
 	{
 		return false;
 	}
 
 	if (const FAuxData* Data = Info->GetAuxData(Index))
 	{
-		Out = FStringView((const TCHAR*)(Data->Data), Data->DataSize / sizeof(TCHAR));
+		Out = FWideStringView((const WIDECHAR*)(Data->Data), Data->DataSize / sizeof(WIDECHAR));
 		return true;
 	}
 
-	Out = FStringView();
+	Out = FWideStringView();
 	return true;
 }
 
@@ -744,9 +744,9 @@ bool IAnalyzer::FEventData::GetString(const ANSICHAR* FieldName, FString& Out) c
 			return true;
 		}
 
-		if (Field.SizeAndType == sizeof(TCHAR))
+		if (Field.SizeAndType == sizeof(WIDECHAR))
 		{
-			Out = FStringView((const TCHAR*)(Data->Data), Data->DataSize / sizeof(TCHAR));
+			Out = FWideStringView((const WIDECHAR*)(Data->Data), Data->DataSize / sizeof(WIDECHAR));
 			return true;
 		}
 	}
diff --git a/Engine/Source/Developer/TraceAnalysis/Public/Trace/Analyzer.h b/Engine/Source/Developer/TraceAnalysis/Public/Trace/Analyzer.h
--- a/Engine/Source/Developer/TraceAnalysis/Public/Trace/Analyzer.h
+++ b/Engine/Source/Developer/TraceAnalysis/Public/Trace/Analyzer.h
@@ -153,7 +153,7 @@ public:
 		  * @param Out Destination object for the field's value.
 		  * @return True if the field was found. */
 		bool GetString(const ANSICHAR* FieldName, FAnsiStringView& Out) const;
-		bool GetString(const ANSICHAR* FieldName, FStringView& Out) const;
+		bool GetString(const ANSICHAR* FieldName, FWideStringView& Out) const;
 		bool GetString(const ANSICHAR* FieldName, FString& Out) const;
 
 		/** The size of the event in uncompressed bytes excluding the header */
