Read much larger chunks of trace data to improve throughput and ease analysis engine's poor behaviour when it is waiting on a dependent event to arrive.

#rnx

diff --git a/Engine/Source/Developer/TraceAnalysis/Private/Analysis/Processor.cpp b/Engine/Source/Developer/TraceAnalysis/Private/Analysis/Processor.cpp
--- a/Engine/Source/Developer/TraceAnalysis/Private/Analysis/Processor.cpp
+++ b/Engine/Source/Developer/TraceAnalysis/Private/Analysis/Processor.cpp
@@ -35,7 +35,7 @@ uint32 FAnalysisProcessor::FImpl::Run()
 {
 	AnalysisEngine.Begin();
 
-	FStreamBuffer Buffer;
+	FStreamBuffer Buffer(4 << 20);
 	while (!StopEvent->Wait(0, true))
 	{
 		UnpausedEvent->Wait();
diff --git a/Engine/Source/Developer/TraceAnalysis/Private/Analysis/StreamReader.cpp b/Engine/Source/Developer/TraceAnalysis/Private/Analysis/StreamReader.cpp
--- a/Engine/Source/Developer/TraceAnalysis/Private/Analysis/StreamReader.cpp
+++ b/Engine/Source/Developer/TraceAnalysis/Private/Analysis/StreamReader.cpp
@@ -73,6 +73,13 @@ void FStreamReader::RestoreMark(struct FMark* Mark)
 
 
 
+////////////////////////////////////////////////////////////////////////////////
+FStreamBuffer::FStreamBuffer(uint32 InitialBufferSize)
+: BufferSize(InitialBufferSize)
+{
+	Buffer = (uint8*)FMemory::Malloc(BufferSize);
+}
+
 ////////////////////////////////////////////////////////////////////////////////
 void FStreamBuffer::Append(const uint8* Data, uint32 Size)
 {
diff --git a/Engine/Source/Developer/TraceAnalysis/Private/Analysis/StreamReader.h b/Engine/Source/Developer/TraceAnalysis/Private/Analysis/StreamReader.h
--- a/Engine/Source/Developer/TraceAnalysis/Private/Analysis/StreamReader.h
+++ b/Engine/Source/Developer/TraceAnalysis/Private/Analysis/StreamReader.h
@@ -62,6 +62,7 @@ class FStreamBuffer
 {
 public:
 								FStreamBuffer() = default;
+								FStreamBuffer(uint32_t InitalBufferSize);
 								~FStreamBuffer();
 								FStreamBuffer(FStreamBuffer&& Rhs) noexcept;
 								FStreamBuffer(const FStreamBuffer&)	= default;
