Restored early out once all analyzers have indicated they're no longer interested in further events.

#rnx

diff --git a/Engine/Source/Developer/TraceAnalysis/Private/Analysis/Engine.cpp b/Engine/Source/Developer/TraceAnalysis/Private/Analysis/Engine.cpp
--- a/Engine/Source/Developer/TraceAnalysis/Private/Analysis/Engine.cpp
+++ b/Engine/Source/Developer/TraceAnalysis/Private/Analysis/Engine.cpp
@@ -926,6 +926,7 @@ class FAnalyzerHub
 public:
 	void				End();
 	void				SetAnalyzers(TArray<IAnalyzer*>&& InAnalyzers);
+	uint32				GetActiveAnalyzerNum() const;
 	void				OnNewType(const FTypeRegistry::FTypeInfo* TypeInfo);
 	void				OnEvent(const FTypeRegistry::FTypeInfo& TypeInfo, IAnalyzer::EStyle Style, const IAnalyzer::FOnEventContext& Context);
 	void				OnThreadInfo(const FThreads::FInfo& ThreadInfo);
@@ -981,6 +982,17 @@ void FAnalyzerHub::SetAnalyzers(TArray<IAnalyzer*>&& InAnalyzers)
 	BuildRoutes();
 }
 
+////////////////////////////////////////////////////////////////////////////////
+uint32 FAnalyzerHub::GetActiveAnalyzerNum() const
+{
+	uint32 Count = 0;
+	for (IAnalyzer* Analyzer : Analyzers)
+	{
+		Count += (Analyzer != nullptr);
+	}
+	return Count;
+}
+
 ////////////////////////////////////////////////////////////////////////////////
 int32 FAnalyzerHub::GetRouteIndex(const FTypeRegistry::FTypeInfo& TypeInfo)
 {
@@ -1458,6 +1470,7 @@ public:
 	typedef FAnalysisState::FSerial FSerial;
 
 						FAnalysisBridge(TArray<IAnalyzer*>&& Analyzers);
+	bool				IsStillAnalyzing() const;
 	void				Reset();
 	uint32				GetUserUidBias() const;
 	FSerial&			GetSerial();
@@ -1486,6 +1499,13 @@ FAnalysisBridge::FAnalysisBridge(TArray<IAnalyzer*>&& Analyzers)
 	AnalyzerHub.SetAnalyzers(MoveTemp(TempAnalyzers));
 }
 
+////////////////////////////////////////////////////////////////////////////////
+bool FAnalysisBridge::IsStillAnalyzing() const
+{
+	// "> 1" because TraceAnalyzer is always present but shouldn't be considered
+	return AnalyzerHub.GetActiveAnalyzerNum() > 1;
+}
+
 ////////////////////////////////////////////////////////////////////////////////
 void FAnalysisBridge::Reset()
 {
@@ -3597,7 +3617,9 @@ void FAnalysisEngine::FImpl::End()
 ////////////////////////////////////////////////////////////////////////////////
 bool FAnalysisEngine::FImpl::OnData(FStreamReader& Reader)
 {
-	return (Machine.OnData(Reader) != FAnalysisMachine::EStatus::Error);
+	bool bRet = (Machine.OnData(Reader) != FAnalysisMachine::EStatus::Error);
+	bRet &= Bridge.IsStillAnalyzing();
+	return bRet;
 }
 
 
