Call OnData() one last time after reaching the end of the input trace stream.

Newer protocol traces consume as much data as they can before returning from OnData(). Older protocol support however can return will data still to process that will be continued on the next OnData() call. However, if the end of input stream has already been met the subsequent OnData() call is not made. The bug's always been present but hidden by a small input buffer.

#rnx

diff --git a/Engine/Source/Developer/TraceAnalysis/Private/Analysis/Processor.cpp b/Engine/Source/Developer/TraceAnalysis/Private/Analysis/Processor.cpp
--- a/Engine/Source/Developer/TraceAnalysis/Private/Analysis/Processor.cpp
+++ b/Engine/Source/Developer/TraceAnalysis/Private/Analysis/Processor.cpp
@@ -45,8 +45,9 @@ uint32 FAnalysisProcessor::FImpl::Run()
 			return DataStream.Read(Out, Size);
 		});
 
-		if (BytesRead <= 0)
+		if (BytesRead < 0)
 		{
+			// An error occurred
 			break;
 		}
 
@@ -54,6 +55,12 @@ uint32 FAnalysisProcessor::FImpl::Run()
 		{
 			break;
 		}
+
+		if (BytesRead == 0)
+		{
+			// There's no more data
+			break;
+		}
 	}
 
 	AnalysisEngine.End();
