Backed out 18978509

#rnx

diff --git a/Engine/Source/Developer/TraceAnalysis/Private/Analysis/Processor.cpp b/Engine/Source/Developer/TraceAnalysis/Private/Analysis/Processor.cpp
--- a/Engine/Source/Developer/TraceAnalysis/Private/Analysis/Processor.cpp
+++ b/Engine/Source/Developer/TraceAnalysis/Private/Analysis/Processor.cpp
@@ -45,9 +45,8 @@ uint32 FAnalysisProcessor::FImpl::Run()
 			return DataStream.Read(Out, Size);
 		});
 
-		if (BytesRead < 0)
+		if (BytesRead <= 0)
 		{
-			// An error occurred
 			break;
 		}
 
@@ -55,12 +54,6 @@ uint32 FAnalysisProcessor::FImpl::Run()
 		{
 			break;
 		}
-
-		if (BytesRead == 0)
-		{
-			// There's no more data
-			break;
-		}
 	}
 
 	AnalysisEngine.End();
