[Insights] Fix trace list not refresing when adding or deleting trace files on Linux.

diff --git a/Engine/Source/Developer/TraceAnalysis/Private/Store/AsioStore.cpp b/Engine/Source/Developer/TraceAnalysis/Private/Store/AsioStore.cpp
--- a/Engine/Source/Developer/TraceAnalysis/Private/Store/AsioStore.cpp
+++ b/Engine/Source/Developer/TraceAnalysis/Private/Store/AsioStore.cpp
@@ -11,7 +11,9 @@
 #	include "Windows/AllowWindowsPlatformTypes.h"
 #	include <Windows.h>
 #	include "Windows/HideWindowsPlatformTypes.h"
-#endif // PLATFORM_WINDOWS
+#elif PLATFORM_LINUX
+	#include <sys/inotify.h>
+#endif
 
 namespace UE {
 namespace Trace {
@@ -24,9 +26,20 @@ class FAsioStore::FDirWatcher
 public:
 	using asio::windows::object_handle::object_handle;
 };
+#elif PLATFORM_LINUX
+class FAsioStore::FDirWatcher
+	: public asio::posix::stream_descriptor
+{
+	typedef TFunction<void(asio::error_code error)> Handler;
+public:
+	using asio::posix::stream_descriptor::stream_descriptor;
+	void async_wait(Handler handler)
+	{
+		asio::posix::stream_descriptor::async_wait(asio::posix::stream_descriptor::wait_read, handler);
+	}
+};
 #else
 class FAsioStore::FDirWatcher
-	//: public asio::posix::stream_descriptor
 {
 public:
 	void async_wait(...) {}
@@ -34,7 +47,7 @@ public:
 	void close() {}
 	bool is_open() { return false; }
 };
-#endif // PLATFORM_WINDOWS
+#endif
 
 
 
@@ -159,6 +172,10 @@ FAsioStore::FAsioStore(asio::io_context& InIoContext, const TCHAR* InStoreDir)
 		DirWatchHandle = 0;
 	}
 	DirWatcher = new FDirWatcher(IoContext, DirWatchHandle);
+#elif PLATFORM_LINUX
+	int inotfd = inotify_init();
+	int watch_desc = inotify_add_watch(inotfd, TCHAR_TO_UTF8(InStoreDir), IN_CREATE | IN_DELETE);
+	DirWatcher = new FDirWatcher(IoContext, inotfd);
 #endif
 
 	WatchDir();
