[Insights] Fix trace list not refreshing MAC

diff --git a/Engine/Source/Developer/TraceAnalysis/Private/Store/AsioStore.cpp b/Engine/Source/Developer/TraceAnalysis/Private/Store/AsioStore.cpp
--- a/Engine/Source/Developer/TraceAnalysis/Private/Store/AsioStore.cpp
+++ b/Engine/Source/Developer/TraceAnalysis/Private/Store/AsioStore.cpp
@@ -30,14 +30,106 @@ public:
 class FAsioStore::FDirWatcher
 	: public asio::posix::stream_descriptor
 {
-	typedef TFunction<void(asio::error_code error)> Handler;
+	typedef TFunction<void(asio::error_code error)> HandlerType;
 public:
 	using asio::posix::stream_descriptor::stream_descriptor;
-	void async_wait(Handler handler)
+	void async_wait(HandlerType InHandler)
 	{
-		asio::posix::stream_descriptor::async_wait(asio::posix::stream_descriptor::wait_read, handler);
+		asio::posix::stream_descriptor::async_wait(asio::posix::stream_descriptor::wait_read, InHandler);
 	}
 };
+#elif PLATFORM_MAC
+class FAsioStore::FDirWatcher
+{
+	typedef TFunction<void(asio::error_code error)> HandlerType;
+public:
+	FDirWatcher(const TCHAR* InStoreDir)
+		: StoreDir(InStoreDir)
+	{
+
+	}
+	void async_wait(HandlerType InHandler);
+	void cancel()
+	{
+		close();
+	}
+	void close()
+	{
+		if (bIsRunning)
+		{
+			FSEventStreamStop(EventStream);
+			FSEventStreamUnscheduleFromRunLoop(EventStream, CFRunLoopGetCurrent(), kCFRunLoopDefaultMode);
+			FSEventStreamInvalidate(EventStream);
+			FSEventStreamRelease(EventStream);
+			bIsRunning = false;
+		}
+	}
+	bool is_open() { return bIsRunning; }
+
+	void ProcessChanges(size_t EventCount, void* EventPaths, const FSEventStreamEventFlags EventFlags[])
+	{
+		Handler(asio::error_code());
+	}
+
+	FSEventStreamRef	EventStream;
+    HandlerType Handler;
+private:
+	bool bIsRunning = false;
+	const TCHAR* StoreDir = nullptr;
+};
+
+void MacCallback(ConstFSEventStreamRef StreamRef,
+					void* InDirWatcherPtr,
+					size_t EventCount,
+					void* EventPaths,
+					const FSEventStreamEventFlags EventFlags[],
+					const FSEventStreamEventId EventIDs[])
+{
+	FAsioStore::FDirWatcher* DirWatcherPtr = (FAsioStore::FDirWatcher*)InDirWatcherPtr;
+	check(DirWatcherPtr);
+	check(DirWatcherPtr->EventStream == StreamRef);
+
+	DirWatcherPtr->ProcessChanges(EventCount, EventPaths, EventFlags);
+}
+
+void FAsioStore::FDirWatcher::async_wait(HandlerType InHandler)
+{
+    if (bIsRunning)
+    {
+        return;
+    }
+
+	CFAbsoluteTime Latency = 0.2;	// seconds
+
+	FSEventStreamContext Context;
+	Context.version = 0;
+	Context.info = this;
+	Context.retain = NULL;
+	Context.release = NULL;
+	Context.copyDescription = NULL;
+
+	// Make sure the path is absolute
+	const FString FullPath = FPaths::ConvertRelativePathToFull(StoreDir);
+
+	// Set up streaming and turn it on
+	CFStringRef FullPathMac = FPlatformString::TCHARToCFString(*FullPath);
+	CFArrayRef PathsToWatch = CFArrayCreate(NULL, (const void**)&FullPathMac, 1, NULL);
+
+	EventStream = FSEventStreamCreate(NULL,
+		&MacCallback,
+		&Context,
+		PathsToWatch,
+		kFSEventStreamEventIdSinceNow,
+		Latency,
+		kFSEventStreamCreateFlagUseCFTypes | kFSEventStreamCreateFlagNoDefer
+	);
+
+	FSEventStreamScheduleWithRunLoop(EventStream, CFRunLoopGetCurrent(), kCFRunLoopDefaultMode);
+	FSEventStreamStart(EventStream);
+	bIsRunning = true;
+
+	Handler = InHandler;
+}
 #else
 class FAsioStore::FDirWatcher
 {
@@ -176,6 +268,8 @@ FAsioStore::FAsioStore(asio::io_context& InIoContext, const TCHAR* InStoreDir)
 	int inotfd = inotify_init();
 	int watch_desc = inotify_add_watch(inotfd, TCHAR_TO_UTF8(InStoreDir), IN_CREATE | IN_DELETE);
 	DirWatcher = new FDirWatcher(IoContext, inotfd);
+#elif PLATFORM_MAC
+	DirWatcher = new FDirWatcher(*StoreDir);
 #endif
 
 	WatchDir();
