Fixed mismatched malloc/delete[] pairing.

#rnx

diff --git a/Engine/Source/Developer/TraceAnalysis/Private/Analysis/Transport/PacketTransport.cpp b/Engine/Source/Developer/TraceAnalysis/Private/Analysis/Transport/PacketTransport.cpp
--- a/Engine/Source/Developer/TraceAnalysis/Private/Analysis/Transport/PacketTransport.cpp
+++ b/Engine/Source/Developer/TraceAnalysis/Private/Analysis/Transport/PacketTransport.cpp
@@ -41,7 +41,7 @@ FPacketTransport::~FPacketTransport()
 		for (FPacketNode* Node = Root; Node != nullptr;)
 		{
 			FPacketNode* Next = Node->Next;
-			delete[] Node;
+			FMemory::Free(Node);
 			Node = Next;
 		}
 	}
