[TraceLog] Add definition event types and reference fields.

Adds a new field type for trace events which contains a reference to a previously emitted event. In order to create a reference, the referenced event type needs to be declared with a new "definition" flag. The definition flag also defines type of id that is used for the reference, 8, 16, 32 and 64 bit values are supported.

diff --git a/Engine/Source/Developer/TraceAnalysis/Private/Analysis/Engine.cpp b/Engine/Source/Developer/TraceAnalysis/Private/Analysis/Engine.cpp
--- a/Engine/Source/Developer/TraceAnalysis/Private/Analysis/Engine.cpp
+++ b/Engine/Source/Developer/TraceAnalysis/Private/Analysis/Engine.cpp
@@ -266,6 +266,7 @@ struct FDispatch
 		Flag_Important		= 1 << 0,
 		Flag_MaybeHasAux	= 1 << 1,
 		Flag_NoSync			= 1 << 2,
+		Flag_Definition		= 1 << 3,
 	};
 
 	struct FField
@@ -277,6 +278,7 @@ struct FDispatch
 		int8		SizeAndType;		// value == byte_size, sign == float < 0 < int
 		uint8		Class : 7;
 		uint8		bArray : 1;
+		uint32		RefUid;			// If reference field, uuid of ref type
 	};
 
 	int32			GetFieldIndex(const ANSICHAR* Name) const;
@@ -322,6 +324,7 @@ public:
 	void				SetImportant();
 	void				SetMaybeHasAux();
 	void				SetNoSync();
+	void				SetDefinition();
 	FDispatch::FField&	AddField(const ANSICHAR* Name, int32 NameSize, uint16 Size);
 	FDispatch*			Finalize();
 
@@ -412,6 +415,13 @@ void FDispatchBuilder::SetNoSync()
 	Dispatch->Flags |= FDispatch::Flag_NoSync;
 }
 
+////////////////////////////////////////////////////////////////////////////////
+void FDispatchBuilder::SetDefinition()
+{
+	auto* Dispatch = (FDispatch*)(Buffer.GetData());
+	Dispatch->Flags |= FDispatch::Flag_Definition;
+}
+
 ////////////////////////////////////////////////////////////////////////////////
 FDispatch::FField& FDispatchBuilder::AddField(const ANSICHAR* Name, int32 NameSize, uint16 Size)
 {
@@ -419,6 +429,7 @@ FDispatch::FField& FDispatchBuilder::AddField(const ANSICHAR* Name, int32 NameSi
 	auto* Field = (FDispatch::FField*)(Buffer.GetData() + Bufoff);
 	Field->NameOffset = AppendName(Name, NameSize);
 	Field->Size = Size;
+	Field->RefUid = 0;
 
 	FFnv1aHash Hash;
 	Hash.Add((const uint8*)Name, NameSize);
@@ -530,6 +541,18 @@ IAnalyzer::FEventFieldInfo::EType IAnalyzer::FEventFieldInfo::GetType() const
 {
 	const auto* Inner = (const FDispatch::FField*)this;
 
+	if (Inner->RefUid != 0)
+	{
+		switch(Inner->Size)
+		{
+			case sizeof(uint8): return EType::Reference8;
+			case sizeof(uint16): return EType::Reference16;
+			case sizeof(uint32): return EType::Reference32;
+			case sizeof(uint64): return EType::Reference64;
+			default: check(false); // Unsupported width
+		}
+	}
+	
 	if (Inner->Class == UE::Trace::Protocol0::Field_String)
 	{
 		return (Inner->SizeAndType == 1) ? EType::AnsiString : EType::WideString;
@@ -759,9 +782,12 @@ uint32 IAnalyzer::FEventData::GetSize() const
 {
 	const auto* Info = (const FEventDataInfo*)this;
 	uint32 Size = Info->Size;
-	for (const FAuxData& Data : *(Info->AuxCollector))
+	if (Info->AuxCollector)
 	{
-		Size += Data.DataSize;
+		for (const FAuxData& Data : *(Info->AuxCollector))
+		{
+			Size += Data.DataSize;
+		}
 	}
 	return Size;
 }
@@ -794,6 +820,37 @@ uint32 IAnalyzer::FEventData::GetAttachmentSize() const
 	const auto* Info = (const FEventDataInfo*)this;
 	return Info->Size - Info->Dispatch.EventSize;
 }
+
+////////////////////////////////////////////////////////////////////////////////
+bool IAnalyzer::FEventData::IsDefinitionImpl(uint32& OutTypeId) const
+{
+	const auto* Info = (const FEventDataInfo*)this;
+	OutTypeId = Info->Dispatch.Uid;
+	return (Info->Dispatch.Flags & FDispatch::Flag_Definition) != 0;
+}
+
+////////////////////////////////////////////////////////////////////////////////
+const void* IAnalyzer::FEventData::GetReferenceValueImpl(const char* FieldName, uint16& OutSizeType, uint32& OutTypeUid) const
+{
+	const auto* Info = (const FEventDataInfo*)this;
+	const auto& Dispatch = Info->Dispatch;
+
+	const int32 Index = Dispatch.GetFieldIndex(FieldName);
+	if (Index < 0)
+	{
+		return nullptr;
+	}
+
+	const auto& Field = Dispatch.Fields[Index];
+	if (Field.RefUid)
+	{
+		OutSizeType = Field.SizeAndType;
+		OutTypeUid = Field.RefUid;
+		return (Info->Ptr + Field.Offset);
+	}
+	return nullptr;
+}
+
 // }}}
 
 
@@ -805,7 +862,9 @@ public:
 	typedef FDispatch FTypeInfo;
 
 						~FTypeRegistry();
-	const FTypeInfo*	Add(const void* TraceData);
+	const FTypeInfo*	Add(const void* TraceData, uint32 Version);
+	const FTypeInfo*	AddVersion4(const void* TraceData);
+	const FTypeInfo*	AddVersion6(const void* TraceData);
 	void				Add(FTypeInfo* TypeInfo);
 	const FTypeInfo*	Get(uint32 Uid) const;
 	bool				IsUidValid(uint32 Uid) const;
@@ -823,6 +882,20 @@ FTypeRegistry::~FTypeRegistry()
 	}
 }
 
+////////////////////////////////////////////////////////////////////////////////
+const FTypeRegistry::FTypeInfo* FTypeRegistry::Add(const void* TraceData, uint32 Version)
+{
+	switch(Version)
+	{
+	case 0: return AddVersion4(TraceData);
+	case 4: return AddVersion4(TraceData);
+	case 6: return AddVersion6(TraceData);
+	default:
+		check(false); // Unknown version
+	}
+	return nullptr;
+}
+
 ////////////////////////////////////////////////////////////////////////////////
 const FTypeRegistry::FTypeInfo* FTypeRegistry::Get(uint32 Uid) const
 {
@@ -835,7 +908,7 @@ const FTypeRegistry::FTypeInfo* FTypeRegistry::Get(uint32 Uid) const
 }
 
 ////////////////////////////////////////////////////////////////////////////////
-const FTypeRegistry::FTypeInfo* FTypeRegistry::Add(const void* TraceData)
+const FTypeRegistry::FTypeInfo* FTypeRegistry::AddVersion4(const void* TraceData)
 {
 	FDispatchBuilder Builder;
 
@@ -892,6 +965,108 @@ const FTypeRegistry::FTypeInfo* FTypeRegistry::Add(const void* TraceData)
 	return TypeInfo;
 }
 
+////////////////////////////////////////////////////////////////////////////////
+const FTypeRegistry::FTypeInfo* FTypeRegistry::AddVersion6(const void* TraceData)
+{
+	using namespace Protocol6;
+	FDispatchBuilder Builder;
+
+	const auto& NewEvent = *(FNewEventEvent*)(TraceData);
+
+	const auto* NameCursor = (const ANSICHAR*)(NewEvent.Fields + NewEvent.FieldCount);
+
+	Builder.SetLoggerName(NameCursor, NewEvent.LoggerNameSize);
+	NameCursor += NewEvent.LoggerNameSize;
+
+	Builder.SetEventName(NameCursor, NewEvent.EventNameSize);
+	NameCursor += NewEvent.EventNameSize;
+	Builder.SetUid(NewEvent.EventUid);
+
+	// Fill out the fields
+	for (int32 i = 0, n = NewEvent.FieldCount; i < n; ++i)
+	{
+		const auto& Field = NewEvent.Fields[i];
+		if (Field.FieldType == EFieldFamily::Regular)
+		{
+			uint16 TypeSize = 1 << (Field.Regular.TypeInfo & Protocol0::Field_Pow2SizeMask);
+			if (Field.Regular.TypeInfo & Protocol0::Field_Float)
+			{
+				TypeSize = -TypeSize;
+			}
+
+			auto& OutField = Builder.AddField(NameCursor, Field.Regular.NameSize, Field.Regular.Size);
+			OutField.Offset = Field.Regular.Offset;
+			OutField.SizeAndType = TypeSize;
+
+			OutField.Class = Field.Regular.TypeInfo & Protocol0::Field_SpecialMask;
+			OutField.bArray = !!(Field.Regular.TypeInfo & Protocol0::Field_Array);
+
+			NameCursor += Field.Regular.NameSize;
+			/*if (Field.FieldType == Protocol6::EFieldFamily::Definition)
+			{
+				Builder.SetIsDefinition();
+			}*/
+		}
+		else if (Field.FieldType == EFieldFamily::Reference)
+		{
+			check((Field.Reference.TypeInfo & Protocol0::Field_CategoryMask) == Protocol0::Field_Integer);
+			const uint16 TypeSize = 1 << (Field.Reference.TypeInfo & Protocol0::Field_Pow2SizeMask);
+
+			auto& OutField = Builder.AddField(NameCursor, Field.Reference.NameSize, TypeSize);
+			OutField.Offset = Field.Reference.Offset;
+			OutField.SizeAndType = TypeSize;
+			OutField.RefUid = uint32(Field.Reference.RefUid) /*<< EKnownEventUids::_UidShift*/;
+			OutField.Class = 0;
+			OutField.bArray = false;
+
+			NameCursor += Field.Reference.NameSize;
+		}
+		else if (Field.FieldType == EFieldFamily::DefinitionId)
+		{
+			check((Field.DefinitionId.TypeInfo & Protocol0::Field_CategoryMask) == Protocol0::Field_Integer);
+			const uint16 TypeSize = 1 << (Field.DefinitionId.TypeInfo & Protocol0::Field_Pow2SizeMask);
+			
+			auto DefinitionIdFieldName = ANSITEXTVIEW("DefinitionId");
+			auto& OutField = Builder.AddField(DefinitionIdFieldName.GetData(), DefinitionIdFieldName.Len(), TypeSize);
+			OutField.Offset = Field.DefinitionId.Offset;
+			OutField.SizeAndType = TypeSize;
+			OutField.RefUid = NewEvent.EventUid;
+			OutField.Class = 0;
+			OutField.bArray = false;
+		}
+		else
+		{
+			// Error!
+			continue;
+		}
+	}
+
+	if (NewEvent.Flags & uint32(EEventFlags::Important))
+	{
+		Builder.SetImportant();
+	}
+
+	if (NewEvent.Flags & uint32(EEventFlags::MaybeHasAux))
+	{
+		Builder.SetMaybeHasAux();
+	}
+
+	if (NewEvent.Flags & uint32(EEventFlags::NoSync))
+	{
+		Builder.SetNoSync();
+	}
+
+	if (NewEvent.Flags & uint32(EEventFlags::Definition))
+	{
+		Builder.SetDefinition();
+	}
+
+	FTypeInfo* TypeInfo = Builder.Finalize();
+	Add(TypeInfo);
+
+	return TypeInfo;
+}
+
 ////////////////////////////////////////////////////////////////////////////////
 void FTypeRegistry::Add(FTypeInfo* TypeInfo)
 {
@@ -1816,7 +1991,7 @@ FProtocol0Stage::EStatus FProtocol0Stage::OnData(
 		{
 			// There is no need to check size here as the runtime never builds
 			// packets that fragment new-event events.
-			const FTypeRegistry::FTypeInfo* TypeInfo = TypeRegistry.Add(Header->EventData);
+			const FTypeRegistry::FTypeInfo* TypeInfo = TypeRegistry.Add(Header->EventData, 0);
 			Context.Bridge.OnNewType(TypeInfo);
 			Transport->Advance(BlockSize);
 			continue;
@@ -2179,7 +2354,7 @@ int32 FProtocol2Stage::OnData(FStreamReader& Reader, FAnalysisBridge& Bridge)
 		{
 			// There is no need to check size here as the runtime never builds
 			// packets that fragment new-event events.
-			TypeInfo = TypeRegistry.Add(EventData);
+			TypeInfo = TypeRegistry.Add(EventData, 0);
 			Bridge.OnNewType(TypeInfo);
 		}
 		else
@@ -2422,7 +2597,7 @@ int32 FProtocol4Stage::OnDataKnown(
 	case Protocol4::EKnownEventUids::NewEvent:
 		{
 			const auto* Size = Reader.GetPointer<uint16>();
-			const FTypeRegistry::FTypeInfo* TypeInfo = TypeRegistry.Add(Size + 1);
+			const FTypeRegistry::FTypeInfo* TypeInfo = TypeRegistry.Add(Size + 1, 4);
 			Bridge.OnNewType(TypeInfo);
 			Reader.Advance(sizeof(*Size) + *Size);
 			return 1;
@@ -2496,7 +2671,7 @@ class FProtocol5Stage
 public:
 							FProtocol5Stage(FTransport* InTransport);
 
-private:
+protected:
 	struct alignas(16) FEventDesc
 	{
 		union
@@ -2568,6 +2743,7 @@ private:
 	EventDescArray			SerialGaps;
 	uint32					NextSerial = ~0u;
 	uint32					SyncCount;
+	uint32					EventVersion = 4; //Protocol version 5 uses the event version from protocol 4
 };
 
 ////////////////////////////////////////////////////////////////////////////////
@@ -2646,7 +2822,7 @@ FProtocol5Stage::EStatus FProtocol5Stage::OnDataNewEvents(const FMachineContext&
 
 	for (const FEventDesc& EventDesc : EventDescs)
 	{
-		const FTypeRegistry::FTypeInfo* TypeInfo = TypeRegistry.Add(EventDesc.Data);
+		const FTypeRegistry::FTypeInfo* TypeInfo = TypeRegistry.Add(EventDesc.Data, EventVersion);
 		Context.Bridge.OnNewType(TypeInfo);
 	}
 
@@ -3288,19 +3464,30 @@ void FProtocol5Stage::DetectSerialGaps(TArray<FEventDescStream>& EventDescHeap)
 		uint32 GapCount = 0;
 		auto RecordGap = [this, &GapCount] (int32 Lhs, int32 Rhs)
 		{
-			for (uint32 n = SerialGaps.Num(); GapCount < n; ++GapCount)
+			if (SerialGaps.IsEmpty())
 			{
-				FEventDesc& SerialGap = SerialGaps[GapCount];
+				return false;
+			}
+			
+			const FEventDesc& SerialGap = SerialGaps[GapCount];
 
-				if (SerialGap.Serial >= Lhs)
-				{
-					GapCount += (SerialGap.Serial == Lhs);
-					break;
-				}
+			if (SerialGap.Serial == Lhs)
+			{
+				/* This is the expected case */
+				++GapCount;
+				return true;
+			}
 
+			if (SerialGap.Serial > Lhs)
+			{
+				/* We've started receiving new gaps that are exist because not all
+				* data has been received yet. They're false positives. No need to process
+				* any further */
 				return false;
 			}
-			return true;
+
+			// If we're here something's probably gone wrong
+			return false;
 		};
 
 		ForEachSerialGap(EventDescHeap, RecordGap);
@@ -3414,6 +3601,23 @@ int32 FProtocol5Stage::DispatchEvents(
 }
 
 
+// {{{1 protocol-6 -------------------------------------------------------------
+
+////////////////////////////////////////////////////////////////////////////////
+class FProtocol6Stage : public FProtocol5Stage
+{
+public:
+							FProtocol6Stage(FTransport* InTransport);
+protected:	
+	using EKnownUids		= Protocol6::EKnownEventUids;
+};
+
+////////////////////////////////////////////////////////////////////////////////
+FProtocol6Stage::FProtocol6Stage(FTransport* InTransport)
+	: FProtocol5Stage(InTransport)
+{
+	EventVersion = 6;
+}
 
 // {{{1 est.-transport ---------------------------------------------------------
 
@@ -3476,6 +3680,11 @@ FEstablishTransportStage::EStatus FEstablishTransportStage::OnData(
 		Context.Machine.Transition();
 		break;
 
+	case Protocol6::EProtocol::Id:
+		Context.Machine.QueueStage<FProtocol6Stage>(Transport);
+		Context.Machine.Transition();
+		break;
+
 	default:
 		return EStatus::Error;
 	}
diff --git a/Engine/Source/Developer/TraceAnalysis/Public/Trace/Analyzer.h b/Engine/Source/Developer/TraceAnalysis/Public/Trace/Analyzer.h
--- a/Engine/Source/Developer/TraceAnalysis/Public/Trace/Analyzer.h
+++ b/Engine/Source/Developer/TraceAnalysis/Public/Trace/Analyzer.h
@@ -51,7 +51,7 @@ public:
 
 	struct TRACEANALYSIS_API FEventFieldInfo
 	{
-		enum class EType { None, Integer, Float, AnsiString, WideString };
+		enum class EType { None, Integer, Float, AnsiString, WideString, Reference8, Reference16, Reference32, Reference64 };
 
 		/** Returns the name of the field. */
 		const ANSICHAR* GetName() const;
@@ -156,6 +156,18 @@ public:
 		bool GetString(const ANSICHAR* FieldName, FWideStringView& Out) const;
 		bool GetString(const ANSICHAR* FieldName, FString& Out) const;
 
+		/** Returns a value of a reference field.
+		 * @param FieldName Name of field
+		 * @return Reference value
+		 */
+		template<typename DefinitionType>
+		TEventRef<DefinitionType> GetReferenceValue(const ANSICHAR* FieldName) const;
+
+		/** If this is a spec event, gets the unique Id for this spec.
+		 * @return A valid spec id if the event is valid, otherwise an empty id.
+		 */
+		template<typename DefinitionType> TEventRef<DefinitionType> GetDefinitionId() const;
+
 		/** The size of the event in uncompressed bytes excluding the header */
 		uint32 GetSize() const;
 
@@ -172,6 +184,8 @@ public:
 		uint32 GetAttachmentSize() const;
 
 	private:
+		bool IsDefinitionImpl(uint32& OutTypeId) const;
+		const void* GetReferenceValueImpl(const char* FieldName, uint16& OutSizeType, uint32& OutTypeUid) const;
 		const void* GetValueImpl(const ANSICHAR* FieldName, int16& SizeAndType) const;
 		const FArrayReader* GetArrayImpl(const ANSICHAR* FieldName) const;
 	};
@@ -373,5 +387,32 @@ const ValueType* IAnalyzer::TArrayReader<ValueType>::GetData() const
 	return (const ValueType*)Addr;
 }
 
+////////////////////////////////////////////////////////////////////////////////
+template<typename DefinitionType>
+TEventRef<DefinitionType> IAnalyzer::FEventData::GetDefinitionId() const
+{
+	uint32 TypeUid;
+	if (IsDefinitionImpl(TypeUid))
+	{
+		//todo: Emit warning when trying to access id of incorrect type?
+		return MakeEventRef(GetValue<DefinitionType>("DefinitionId", 0), TypeUid);
+	}
+	return MakeEventRef<DefinitionType>(0,0);
+}
+
+////////////////////////////////////////////////////////////////////////////////
+template<typename DefinitionType>
+TEventRef<DefinitionType> IAnalyzer::FEventData::GetReferenceValue(const ANSICHAR* FieldName) const
+{
+	uint32 TypeUid;
+	uint16 SizeAndType;
+	const void* Value = GetReferenceValueImpl(FieldName, SizeAndType, TypeUid);
+	if (Value)
+	{
+		return MakeEventRef<DefinitionType>(CoerceValue<DefinitionType>(Value, SizeAndType), TypeUid);
+	}
+	return MakeEventRef<DefinitionType>(0,0);
+}
+
 } // namespace Trace
 } // namespace UE
diff --git a/Engine/Source/Runtime/TraceLog/Private/Trace/EventNode.cpp b/Engine/Source/Runtime/TraceLog/Private/Trace/EventNode.cpp
--- a/Engine/Source/Runtime/TraceLog/Private/Trace/EventNode.cpp
+++ b/Engine/Source/Runtime/TraceLog/Private/Trace/EventNode.cpp
@@ -108,6 +108,9 @@ void FEventNode::Describe() const
 {
 	const FLiteralName& LoggerName = Info->LoggerName;
 	const FLiteralName& EventName = Info->EventName;
+	
+	// Definition events adds an extra field containing the definition id to the end.
+	const uint32 DefinitionIdFieldIdx = Info->FieldCount - ((Info->Flags & FEventInfo::DefinitionBits) ? 1 : 0);
 
 	// Calculate the number of fields and size of name data.
 	uint32 NamesSize = LoggerName.Length + EventName.Length;
@@ -133,10 +136,11 @@ void FEventNode::Describe() const
 	Event.EventNameSize = EventName.Length;
 	Event.Flags = 0;
 
-	uint32 Flags = Info->Flags;
+	const uint32 Flags = Info->Flags;
 	if (Flags & FEventInfo::Flag_Important)		Event.Flags |= uint8(EEventFlags::Important);
 	if (Flags & FEventInfo::Flag_MaybeHasAux)	Event.Flags |= uint8(EEventFlags::MaybeHasAux);
 	if (Flags & FEventInfo::Flag_NoSync)		Event.Flags |= uint8(EEventFlags::NoSync);
+	if (Flags & FEventInfo::DefinitionBits)		Event.Flags |= uint8(EEventFlags::Definition);
 
 	// Write details about event's fields
 	Event.FieldCount = uint8(Info->FieldCount);
@@ -144,10 +148,29 @@ void FEventNode::Describe() const
 	{
 		const FFieldDesc& Field = Info->Fields[i];
 		auto& Out = Event.Fields[i];
-		Out.Offset = Field.ValueOffset;
-		Out.Size = Field.ValueSize;
-		Out.TypeInfo = Field.TypeInfo;
-		Out.NameSize = Field.NameSize;
+		if (i == DefinitionIdFieldIdx)
+		{
+			Out.FieldType = EFieldFamily::DefinitionId;
+			Out.DefinitionId.Offset = Field.ValueOffset;
+			Out.DefinitionId.TypeInfo = Field.TypeInfo;
+		}
+		else if (Field.Reference)
+		{
+			Out.FieldType = EFieldFamily::Reference;
+			// todo: What if the referenced type has not been initialized yet?
+			Out.Reference.Offset = Field.ValueOffset;
+			Out.Reference.TypeInfo = Field.TypeInfo;
+			Out.Reference.NameSize = Field.NameSize;
+			Out.Reference.RefUid = uint16(Field.Reference->GetUid()) >> EKnownEventUids::_UidShift;
+		}
+		else
+		{
+			Out.FieldType = EFieldFamily::Regular;
+			Out.Regular.Offset = Field.ValueOffset;
+			Out.Regular.Size = Field.ValueSize;
+			Out.Regular.TypeInfo = Field.TypeInfo;
+			Out.Regular.NameSize = Field.NameSize;
+		}
 	}
 
 	// Write names
diff --git a/Engine/Source/Runtime/TraceLog/Public/Trace/Config.h b/Engine/Source/Runtime/TraceLog/Public/Trace/Config.h
--- a/Engine/Source/Runtime/TraceLog/Public/Trace/Config.h
+++ b/Engine/Source/Runtime/TraceLog/Public/Trace/Config.h
@@ -17,5 +17,5 @@
 #endif
 
 #if UE_TRACE_ENABLED
-#	define TRACE_PRIVATE_PROTOCOL_5
+#	define TRACE_PRIVATE_PROTOCOL_6
 #endif
diff --git a/Engine/Source/Runtime/TraceLog/Public/Trace/Detail/EventNode.h b/Engine/Source/Runtime/TraceLog/Public/Trace/Detail/EventNode.h
--- a/Engine/Source/Runtime/TraceLog/Public/Trace/Detail/EventNode.h
+++ b/Engine/Source/Runtime/TraceLog/Public/Trace/Detail/EventNode.h
@@ -21,6 +21,12 @@ struct FEventInfo
 		Flag_Important		= 1 << 0,
 		Flag_MaybeHasAux	= 1 << 1,
 		Flag_NoSync			= 1 << 2,
+		Flag_Definition8	= 1 << 3,
+		Flag_Definition16	= 1 << 4,
+		Flag_Definition32	= 1 << 5,
+		Flag_Definition64	= 1 << 6,
+
+		DefinitionBits		= Flag_Definition8 | Flag_Definition16 | Flag_Definition32 | Flag_Definition64,
 	};
 
 	FLiteralName			LoggerName;
diff --git a/Engine/Source/Runtime/TraceLog/Public/Trace/Detail/Field.h b/Engine/Source/Runtime/TraceLog/Public/Trace/Detail/Field.h
--- a/Engine/Source/Runtime/TraceLog/Public/Trace/Detail/Field.h
+++ b/Engine/Source/Runtime/TraceLog/Public/Trace/Detail/Field.h
@@ -28,11 +28,16 @@ UE_TRACE_API void Field_WriteStringAnsi(uint32, const ANSICHAR*, int32);
 UE_TRACE_API void Field_WriteStringAnsi(uint32, const WIDECHAR*, int32);
 UE_TRACE_API void Field_WriteStringWide(uint32, const WIDECHAR*, int32);
 
+class FEventNode;
+
 } // namespace Private
 
 ////////////////////////////////////////////////////////////////////////////////
+enum DisabledField {};
+	
 template <typename Type> struct TFieldType;
 
+template <> struct TFieldType<DisabledField>{ enum { Tid = 0,						Size = 0 }; };
 template <> struct TFieldType<bool>			{ enum { Tid = int(EFieldType::Bool),	Size = sizeof(bool) }; };
 template <> struct TFieldType<int8>			{ enum { Tid = int(EFieldType::Int8),	Size = sizeof(int8) }; };
 template <> struct TFieldType<int16>		{ enum { Tid = int(EFieldType::Int16),	Size = sizeof(int16) }; };
@@ -91,12 +96,13 @@ struct FLiteralName
 ////////////////////////////////////////////////////////////////////////////////
 struct FFieldDesc
 {
-	FFieldDesc(const FLiteralName& Name, uint8 Type, uint16 Offset, uint16 Size)
+	FFieldDesc(const FLiteralName& Name, uint8 Type, uint16 Offset, uint16 Size, Private::FEventNode* ReferenceEvent = nullptr)
 	: Name(Name.Ptr)
 	, ValueOffset(Offset)
 	, ValueSize(Size)
 	, NameSize(Name.Length)
 	, TypeInfo(Type)
+	, Reference(ReferenceEvent)
 	{
 	}
 
@@ -105,6 +111,7 @@ struct FFieldDesc
 	uint16			ValueSize;
 	uint8			NameSize;
 	uint8			TypeInfo;
+	Private::FEventNode* Reference;
 };
 
 
@@ -137,7 +144,7 @@ enum class EIndexPack
 		: FieldDesc(Name, Tid, Offset, Size) \
 		{ \
 		}
-
+	
 ////////////////////////////////////////////////////////////////////////////////
 template <int InIndex, int InOffset, typename Type>
 struct TField<InIndex, InOffset, Type[]>
@@ -168,6 +175,18 @@ struct TField<InIndex, InOffset, WideString>
 	TRACE_PRIVATE_FIELD(InIndex + int(EIndexPack::AuxFieldCounter), InOffset, WideString);
 };
 
+////////////////////////////////////////////////////////////////////////////////
+template <int InIndex, int InOffset, typename DefinitionType>
+struct TField<InIndex, InOffset, TEventRef<DefinitionType>>
+{
+	TRACE_PRIVATE_FIELD(InIndex, InOffset, DefinitionType);
+public:
+	TField(const FLiteralName& Name, Private::FEventNode* ReferenceEvent)
+		: FieldDesc(Name, Tid, Offset, Size, ReferenceEvent)
+	{
+	}
+};
+	
 ////////////////////////////////////////////////////////////////////////////////
 template <int InIndex, int InOffset, typename Type>
 struct TField
diff --git a/Engine/Source/Runtime/TraceLog/Public/Trace/Detail/Important/ImportantLogScope.inl b/Engine/Source/Runtime/TraceLog/Public/Trace/Detail/Important/ImportantLogScope.inl
--- a/Engine/Source/Runtime/TraceLog/Public/Trace/Detail/Important/ImportantLogScope.inl
+++ b/Engine/Source/Runtime/TraceLog/Public/Trace/Detail/Important/ImportantLogScope.inl
@@ -200,6 +200,16 @@ struct FImportantLogScope::FFieldSet<FieldMeta, WideString>
 	}
 };
 
+////////////////////////////////////////////////////////////////////////////////
+template <typename FieldMeta, typename DefinitionType>
+struct FImportantLogScope::FFieldSet<FieldMeta, TEventRef<DefinitionType>>
+{
+	static void Impl(FImportantLogScope* Scope, const TEventRef<DefinitionType>& Reference)
+	{
+		FFieldSet<FieldMeta, DefinitionType>::Impl(Scope, Reference.Id);
+	}
+};
+
 } // namespace Private
 } // namespace Trace
 } // namespace UE
diff --git a/Engine/Source/Runtime/TraceLog/Public/Trace/Detail/LogScope.inl b/Engine/Source/Runtime/TraceLog/Public/Trace/Detail/LogScope.inl
--- a/Engine/Source/Runtime/TraceLog/Public/Trace/Detail/LogScope.inl
+++ b/Engine/Source/Runtime/TraceLog/Public/Trace/Detail/LogScope.inl
@@ -306,6 +306,16 @@ struct FLogScope::FFieldSet<FieldMeta, WideString>
 	}
 };
 
+////////////////////////////////////////////////////////////////////////////////
+template <typename FieldMeta, typename DefinitionType>
+struct FLogScope::FFieldSet<FieldMeta, TEventRef<DefinitionType>>
+{
+	static void Impl(FLogScope* Scope, const TEventRef<DefinitionType>& Reference)
+	{
+		FFieldSet<FieldMeta, DefinitionType>::Impl(Scope, Reference.Id);
+	}
+};
+
 } // namespace Private
 } // namespace Trace
 } // namespace UE
diff --git a/Engine/Source/Runtime/TraceLog/Public/Trace/Detail/Protocol.h b/Engine/Source/Runtime/TraceLog/Public/Trace/Detail/Protocol.h
--- a/Engine/Source/Runtime/TraceLog/Public/Trace/Detail/Protocol.h
+++ b/Engine/Source/Runtime/TraceLog/Public/Trace/Detail/Protocol.h
@@ -16,6 +16,7 @@
 #include "Protocols/Protocol3.h"
 #include "Protocols/Protocol4.h"
 #include "Protocols/Protocol5.h"
+#include "Protocols/Protocol6.h"
 
 #if defined(_MSC_VER)
 	#pragma warning(pop)
diff --git a/Engine/Source/Runtime/TraceLog/Public/Trace/Detail/Protocols/Protocol6.h b/Engine/Source/Runtime/TraceLog/Public/Trace/Detail/Protocols/Protocol6.h
new file mode 100644
--- /dev/null
+++ b/Engine/Source/Runtime/TraceLog/Public/Trace/Detail/Protocols/Protocol6.h
@@ -0,0 +1,86 @@
+﻿// Copyright Epic Games, Inc. All Rights Reserved.
+#pragma once
+
+#include "Protocol4.h"
+#include "Protocol5.h"
+
+namespace UE {
+namespace Trace {
+
+#if defined(TRACE_PRIVATE_PROTOCOL_6)
+inline
+#endif
+namespace Protocol6
+{
+
+////////////////////////////////////////////////////////////////////////////////
+enum EProtocol : uint8 { Id = 6 };
+
+////////////////////////////////////////////////////////////////////////////////
+using Protocol5::EFieldType;
+using Protocol5::FEventHeader;
+using Protocol5::FImportantEventHeader;
+using Protocol5::FEventHeaderSync;
+using Protocol5::FAuxHeader;
+using Protocol5::EKnownEventUids;
+
+////////////////////////////////////////////////////////////////////////////////
+enum class EEventFlags : uint8
+{
+	Important		= 1 << 0,
+	MaybeHasAux		= 1 << 1,
+	NoSync			= 1 << 2,
+	Definition		= 1 << 3,
+};
+
+////////////////////////////////////////////////////////////////////////////////
+enum EFieldFamily : uint8
+{
+	Regular,
+	Reference,
+	DefinitionId,
+};
+	
+////////////////////////////////////////////////////////////////////////////////
+struct FNewEventEvent
+{
+	uint16		EventUid;
+	uint8		FieldCount;
+	uint8		Flags;
+	uint8		LoggerNameSize;
+	uint8		EventNameSize;
+	struct
+	{
+		uint8	FieldType;
+		union
+		{
+			struct
+			{
+				uint16	Offset;
+				uint16	Size;
+				uint8	TypeInfo;
+				uint8	NameSize;
+			} Regular;
+			struct
+			{
+				uint16	Offset;
+				uint16	RefUid; 
+				uint8	TypeInfo;
+				uint8	NameSize;
+			} Reference;
+			struct
+			{
+				uint16	Offset;
+				uint16	Unused1;
+				uint8	Unused2;
+				uint8	TypeInfo;
+			} DefinitionId;
+		};
+	} Fields[];
+};
+	
+////////////////////////////////////////////////////////////////////////////////
+
+} // namespace Protocol6
+} // namespace Trace
+} // namespace UE
diff --git a/Engine/Source/Runtime/TraceLog/Public/Trace/Detail/Trace.h b/Engine/Source/Runtime/TraceLog/Public/Trace/Detail/Trace.h
--- a/Engine/Source/Runtime/TraceLog/Public/Trace/Detail/Trace.h
+++ b/Engine/Source/Runtime/TraceLog/Public/Trace/Detail/Trace.h
@@ -65,9 +65,19 @@ class FChannel;
 		{ \
 			Important			= UE::Trace::Private::FEventInfo::Flag_Important, \
 			NoSync				= UE::Trace::Private::FEventInfo::Flag_NoSync, \
+			Definition8bit		= UE::Trace::Private::FEventInfo::Flag_Definition8, \
+			Definition16bit		= UE::Trace::Private::FEventInfo::Flag_Definition16, \
+			Definition32bit		= UE::Trace::Private::FEventInfo::Flag_Definition32, \
+			Definition64bit		= UE::Trace::Private::FEventInfo::Flag_Definition64, \
+			DefinitionBits		= UE::Trace::Private::FEventInfo::DefinitionBits, \
 			PartialEventFlags	= (0, ##__VA_ARGS__), \
 		}; \
-		enum : bool { bIsImportant = ((0, ##__VA_ARGS__) & Important) != 0, }; \
+		enum : bool { bIsImportant = ((0, ##__VA_ARGS__) & Important) != 0, bIsDefinition = ((0, ##__VA_ARGS__) & DefinitionBits) != 0,\
+		bIsDefinition8 = ((0, ##__VA_ARGS__) & Definition8bit) != 0, \
+		bIsDefinition16 = ((0, ##__VA_ARGS__) & Definition16bit) != 0,\
+		bIsDefinition32 = ((0, ##__VA_ARGS__) & Definition32bit) != 0, \
+		bIsDefinition64 = ((0, ##__VA_ARGS__) & Definition64bit) != 0,}; \
+		typedef std::conditional_t<bIsDefinition8, UE::Trace::FEventRef8, std::conditional_t<bIsDefinition16, UE::Trace::FEventRef16 , std::conditional_t<bIsDefinition64, UE::Trace::FEventRef64, UE::Trace::FEventRef32>>> DefinitionType;\
 		static constexpr uint32 GetSize() { return EventProps_Meta::Size; } \
 		static uint32 GetUid() { static uint32 Uid = 0; return (Uid = Uid ? Uid : Initialize()); } \
 		static uint32 FORCENOINLINE Initialize() \
@@ -81,7 +91,7 @@ class FChannel;
 					FLiteralName(#LoggerName), \
 					FLiteralName(#EventName), \
 					(FFieldDesc*)(&Fields), \
-					uint16(sizeof(Fields) / sizeof(FFieldDesc)), \
+					EventProps_Meta::NumFields, \
 					uint16(EventFlags), \
 				}; \
 				return LoggerName##EventName##Event.Initialize(&Info); \
@@ -101,8 +111,29 @@ class FChannel;
 			FieldName##_Meta::Index + 1, \
 			FieldName##_Meta::Offset + FieldName##_Meta::Size,
 
+#define TRACE_PRIVATE_EVENT_REFFIELD(RefLogger, RefEventType, FieldName) \
+		F##RefLogger##RefEventType##Fields::DefinitionType> FieldName##_Meta; \
+		FieldName##_Meta const FieldName##_Field = FieldName##_Meta(UE::Trace::FLiteralName(#FieldName), &RefLogger##RefEventType##Event); \
+		template <typename DefinitionType> auto FieldName(UE::Trace::TEventRef<DefinitionType> Reference) const { \
+			checkfSlow(Reference.RefTypeId == F##RefLogger##RefEventType##Fields::GetUid(), TEXT("Incorrect reference type passed to event. Field expected %s with uid %u but got a reference with uid %u"), TEXT(#RefEventType), F##RefLogger##RefEventType##Fields::GetUid(), Reference.RefTypeId);\
+			LogScopeType::FFieldSet<FieldName##_Meta, F##RefLogger##RefEventType##Fields::DefinitionType>::Impl((LogScopeType*)this, Reference); \
+			return true; \
+		} \
+		typedef UE::Trace::TField< \
+			FieldName##_Meta::Index + 1, \
+			FieldName##_Meta::Offset + FieldName##_Meta::Size,
+
 #define TRACE_PRIVATE_EVENT_END() \
-		UE::Trace::EventProps> EventProps_Meta; \
+		std::conditional<bIsDefinition, DefinitionType, UE::Trace::DisabledField>::type> DefinitionId_Meta;\
+		DefinitionId_Meta const DefinitionId_Field = UE::Trace::FLiteralName("");\
+		static constexpr uint16 NumDefinitionFields = bIsDefinition ? 1 : 0;\
+		template<typename RefType>\
+		auto SetDefinitionId(const RefType& Id) const \
+		{ \
+			LogScopeType::FFieldSet<DefinitionId_Meta, RefType>::Impl((LogScopeType*)this, Id); \
+			return true; \
+		} \
+		typedef UE::Trace::TField<DefinitionId_Meta::Index + NumDefinitionFields, DefinitionId_Meta::Offset + DefinitionId_Meta::Size, UE::Trace::EventProps> EventProps_Meta; \
 		EventProps_Meta const EventProps_Private = {}; \
 		typedef std::conditional<bIsImportant, UE::Trace::Private::FImportantLogScope, UE::Trace::Private::FLogScope>::type LogScopeType; \
 		explicit operator bool () const { return true; } \
@@ -138,6 +169,14 @@ class FChannel;
 		PREPROCESSOR_JOIN(TheScope, __LINE__).SetActive(), \
 		TRACE_PRIVATE_LOG_EPILOG()
 
+#define TRACE_PRIVATE_GET_DEFINITION_TYPE_ID(LoggerName, EventName) \
+	F##LoggerName##EventName##Fields::GetUid()
+
+#define TRACE_PRIVATE_LOG_DEFINITION(LoggerName, EventName, Id, ChannelsExpr, ...) \
+	UE::Trace::MakeEventRef(Id, TRACE_PRIVATE_GET_DEFINITION_TYPE_ID(LoggerName, EventName)); \
+	TRACE_PRIVATE_LOG(LoggerName, EventName, ChannelsExpr, ##__VA_ARGS__) \
+		<< EventName.SetDefinitionId(UE::Trace::MakeEventRef(Id, F##LoggerName##EventName##Fields::GetUid()))
+
 #else
 
 #define TRACE_PRIVATE_CHANNEL(ChannelName, ...)
@@ -173,6 +212,9 @@ class FChannel;
 #define TRACE_PRIVATE_EVENT_FIELD(FieldType, FieldName) \
 		const FTraceDisabled& FieldName;
 
+#define TRACE_PRIVATE_EVENT_REFFIELD(RefLogger, RefEventType, FieldName) \
+		const FTraceDisabled& FieldName;
+
 #define TRACE_PRIVATE_EVENT_END() \
 	};
 
@@ -188,4 +230,11 @@ class FChannel;
 	if (const auto& EventName = *(F##LoggerName##EventName##Dummy*)1) \
 		EventName
 
+#define TRACE_PRIVATE_GET_DEFINITION_TYPE_ID(LoggerName, EventName) \
+	0
+
+#define TRACE_PRIVATE_LOG_DEFINITION(LoggerName, EventName, Id, ChannelsExpr, ...) \
+	UE::Trace::MakeEventRef(Id, 0); \
+	TRACE_PRIVATE_LOG(LoggerName, EventName, ChannelsExpr, ##__VA_ARGS__) \
+
 #endif // UE_TRACE_ENABLED
diff --git a/Engine/Source/Runtime/TraceLog/Public/Trace/Trace.h b/Engine/Source/Runtime/TraceLog/Public/Trace/Trace.h
--- a/Engine/Source/Runtime/TraceLog/Public/Trace/Trace.h
+++ b/Engine/Source/Runtime/TraceLog/Public/Trace/Trace.h
@@ -21,6 +21,62 @@ namespace Trace {
 enum AnsiString {};
 enum WideString {};
 
+// Reference to a definition event.
+template<typename IdType>	
+struct TEventRef 
+{
+	using ReferenceType = IdType;
+	
+	TEventRef(IdType InId, uint32 InTypeId)
+		: Id(InId)
+		, RefTypeId(InTypeId)
+	{
+	}
+
+	IdType Id;
+	uint32 RefTypeId;
+
+	uint64 GetHash() const;
+	
+private:
+	TEventRef() = delete;
+};
+
+template <>
+inline uint64 TEventRef<uint8>::GetHash() const
+{
+	return (uint64(RefTypeId) << 32) | Id;
+}
+	
+template <>
+inline uint64 TEventRef<uint16>::GetHash() const
+{
+	return (uint64(RefTypeId) << 32) | Id;
+}
+	
+template <>
+inline uint64 TEventRef<uint32>::GetHash() const
+{
+	return (uint64(RefTypeId) << 32) | Id;
+}
+
+template <>
+inline uint64 TEventRef<uint64>::GetHash() const
+{
+	return (uint64(RefTypeId) << 32) ^ Id;
+}
+
+typedef TEventRef<uint8> FEventRef8;
+typedef TEventRef<uint16> FEventRef16;	
+typedef TEventRef<uint32> FEventRef32;
+typedef TEventRef<uint64> FEventRef64;
+
+template<typename IdType>
+TEventRef<IdType> MakeEventRef(IdType InId, uint32 InTypeId)
+{
+	return TEventRef<IdType>(InId, InTypeId);
+}
+
 struct FInitializeDesc
 {
 	uint32			TailSizeBytes		= 4 << 20;
@@ -66,10 +122,13 @@ UE_TRACE_API void	ThreadGroupEnd() UE_TRACE_IMPL();
 #define UE_TRACE_EVENT_BEGIN(LoggerName, EventName, ...)				TRACE_PRIVATE_EVENT_BEGIN(LoggerName, EventName, ##__VA_ARGS__)
 #define UE_TRACE_EVENT_BEGIN_EXTERN(LoggerName, EventName, ...)			TRACE_PRIVATE_EVENT_BEGIN_EXTERN(LoggerName, EventName, ##__VA_ARGS__)
 #define UE_TRACE_EVENT_FIELD(FieldType, FieldName)						TRACE_PRIVATE_EVENT_FIELD(FieldType, FieldName)
+#define UE_TRACE_EVENT_REFERENCE_FIELD(RefLogger, RefEvent, FieldName)	TRACE_PRIVATE_EVENT_REFFIELD(RefLogger, RefEvent, FieldName)
 #define UE_TRACE_EVENT_END()											TRACE_PRIVATE_EVENT_END()
 #define UE_TRACE_LOG(LoggerName, EventName, ChannelsExpr, ...)			TRACE_PRIVATE_LOG(LoggerName, EventName, ChannelsExpr, ##__VA_ARGS__)
 #define UE_TRACE_LOG_SCOPED(LoggerName, EventName, ChannelsExpr, ...)	TRACE_PRIVATE_LOG_SCOPED(LoggerName, EventName, ChannelsExpr, ##__VA_ARGS__)
 #define UE_TRACE_LOG_SCOPED_T(LoggerName, EventName, ChannelsExpr, ...)	TRACE_PRIVATE_LOG_SCOPED_T(LoggerName, EventName, ChannelsExpr, ##__VA_ARGS__)
+#define UE_TRACE_GET_DEFINITION_TYPE_ID(LoggerName, EventName)			TRACE_PRIVATE_GET_DEFINITION_TYPE_ID(LoggerName, EventName)
+#define UE_TRACE_LOG_DEFINITION(LoggerName, EventName, Id, ChannelsExpr, ...) TRACE_PRIVATE_LOG_DEFINITION(LoggerName, EventName, Id, ChannelsExpr, ##__VA_ARGS__)
 
 ////////////////////////////////////////////////////////////////////////////////
 #define UE_TRACE_CHANNEL(ChannelName, ...)				TRACE_PRIVATE_CHANNEL(ChannelName, ##__VA_ARGS__)
