[TraceLog] Add signed/unsigned integer distinction

diff --git a/Engine/Source/Developer/TraceAnalysis/Private/Analysis/Engine.cpp b/Engine/Source/Developer/TraceAnalysis/Private/Analysis/Engine.cpp
--- a/Engine/Source/Developer/TraceAnalysis/Private/Analysis/Engine.cpp
+++ b/Engine/Source/Developer/TraceAnalysis/Private/Analysis/Engine.cpp
@@ -52,7 +52,8 @@ struct FAuxData
 	const uint8*	Data;
 	uint32			DataSize;
 	uint16			FieldIndex;
-	int16			FieldSizeAndType;
+	int8			FieldSizeAndType;
+	bool			bSigned;
 };
 
 ////////////////////////////////////////////////////////////////////////////////
@@ -578,6 +579,12 @@ bool IAnalyzer::FEventFieldInfo::IsArray() const
 	return Inner->bArray;
 }
 
+////////////////////////////////////////////////////////////////////////////////
+bool IAnalyzer::FEventFieldInfo::IsSigned() const
+{
+	const auto* Inner = (const FDispatch::FField*)this;
+	return (Inner->Class & UE::Trace::Protocol0::Field_Signed) != 0;
+}
 
 
 // {{{1 array-reader -----------------------------------------------------------
@@ -592,7 +599,7 @@ uint32 IAnalyzer::FArrayReader::Num() const
 }
 
 ////////////////////////////////////////////////////////////////////////////////
-const void* IAnalyzer::FArrayReader::GetImpl(uint32 Index, int16& SizeAndType) const
+const void* IAnalyzer::FArrayReader::GetImpl(uint32 Index, int8& SizeAndType) const
 {
 	const auto* Inner = (const FAuxData*)this;
 	SizeAndType = Inner->FieldSizeAndType;
@@ -633,7 +640,9 @@ const FAuxData* FEventDataInfo::GetAuxData(uint32 FieldIndex) const
 	{
 		if (Data.FieldIndex == FieldIndex)
 		{
-			Data.FieldSizeAndType = Info->Dispatch.Fields[FieldIndex].SizeAndType;
+			const FDispatch::FField& Field = Info->Dispatch.Fields[FieldIndex];
+			Data.FieldSizeAndType = Field.SizeAndType;
+			Data.bSigned = (Field.Class & UE::Trace::Protocol0::Field_Signed) != 0;
 			return &Data;
 		}
 	}
@@ -667,7 +676,7 @@ const IAnalyzer::FArrayReader* IAnalyzer::FEventData::GetArrayImpl(const ANSICHA
 }
 
 ////////////////////////////////////////////////////////////////////////////////
-const void* IAnalyzer::FEventData::GetValueImpl(const ANSICHAR* FieldName, int16& SizeAndType) const
+const void* IAnalyzer::FEventData::GetValueImpl(const ANSICHAR* FieldName, int8& SizeAndType) const
 {
 	const auto* Info = (const FEventDataInfo*)this;
 	const auto& Dispatch = Info->Dispatch;
@@ -2727,6 +2736,7 @@ protected:
 	virtual EStatus			OnData(FStreamReader& Reader, const FMachineContext& Context) override;
 	EStatus					OnDataNewEvents(const FMachineContext& Context);
 	EStatus					OnDataImportant(const FMachineContext& Context);
+	EStatus					OnDataNonCachedImportant(const FMachineContext& Context);
 	EStatus					OnDataNormal(const FMachineContext& Context);
 	int32					ParseImportantEvents(FStreamReader& Reader, EventDescArray& OutEventDescs);
 	int32					ParseEvents(FStreamReader& Reader, EventDescArray& OutEventDescs);
diff --git a/Engine/Source/Developer/TraceAnalysis/Private/Analysis/EventToCbor.cpp b/Engine/Source/Developer/TraceAnalysis/Private/Analysis/EventToCbor.cpp
--- a/Engine/Source/Developer/TraceAnalysis/Private/Analysis/EventToCbor.cpp
+++ b/Engine/Source/Developer/TraceAnalysis/Private/Analysis/EventToCbor.cpp
@@ -8,6 +8,25 @@
 namespace UE {
 namespace Trace {
 
+////////////////////////////////////////////////////////////////////////////////
+template<typename DestType>
+static void WriteArray(FCborWriter& CborWriter, const IAnalyzer::TArrayReader<DestType>& Reader)
+{
+	if (const uint32 Num = Reader.Num())
+	{
+		CborWriter.WriteContainerStart(ECborCode::Array, Reader.Num());
+		for (uint32 j = 0; j < Num; ++j)
+		{
+			DestType Value = Reader[j];
+			CborWriter.WriteValue(Value);
+		}
+	}
+	else
+	{
+		CborWriter.WriteNull();
+	}
+}
+	
 ////////////////////////////////////////////////////////////////////////////////
 void SerializeToCborImpl(
 	TArray<uint8>& Out,
@@ -53,47 +72,34 @@ void SerializeToCborImpl(
 
 		if (FieldInfo.IsArray())
 		{
-			if (FieldInfo.GetType() == IAnalyzer::FEventFieldInfo::EType::Integer)
+			switch (FieldInfo.GetType())
 			{
-				const IAnalyzer::TArrayReader<uint64>& Reader = EventData.GetArray<uint64>(FieldName);
-				if (uint32 Num = Reader.Num())
-				{
-					CborWriter.WriteContainerStart(ECborCode::Array, Reader.Num());
-					for (uint32 j = 0; j < Num; ++j)
-					{
-						uint64 Value = Reader[j];
-						CborWriter.WriteValue(Value);
-					}
-				}
-				else
-				{
-					CborWriter.WriteNull();
-				}
-			}
-			else
-			{
-				const IAnalyzer::TArrayReader<double>& Reader = EventData.GetArray<double>(FieldName);
-				if (uint32 Num = Reader.Num())
-				{
-					CborWriter.WriteContainerStart(ECborCode::Array, Reader.Num());
-					for (uint32 j = 0; j < Num; ++j)
-					{
-						double Value = Reader[j];
-						CborWriter.WriteValue(Value);
-					}
-				}
-				else
-				{
-					CborWriter.WriteNull();
-				}
+			case IAnalyzer::FEventFieldInfo::EType::Integer:
+				WriteArray(CborWriter, EventData.GetArray<uint64>(FieldName));
+				break;
+			case IAnalyzer::FEventFieldInfo::EType::Float:
+				WriteArray(CborWriter, EventData.GetArray<double>(FieldName));
+				break;
+			default:
+				//not supported
+				CborWriter.WriteNull();
+				break;
 			}
 			continue;
 		}
 
 		if (FieldInfo.GetType() == IAnalyzer::FEventFieldInfo::EType::Integer)
 		{
-			uint64 Value = EventData.GetValue<uint64>(FieldName);
-			CborWriter.WriteValue(Value);
+			if (FieldInfo.IsSigned())
+			{
+				int64 Value = EventData.GetValue<int64>(FieldName);
+				CborWriter.WriteValue(Value);
+			}
+			else
+			{
+				uint64 Value = EventData.GetValue<uint64>(FieldName);
+				CborWriter.WriteValue(Value);
+			}
 			continue;
 		}
 
diff --git a/Engine/Source/Developer/TraceAnalysis/Public/Trace/Analyzer.h b/Engine/Source/Developer/TraceAnalysis/Public/Trace/Analyzer.h
--- a/Engine/Source/Developer/TraceAnalysis/Public/Trace/Analyzer.h
+++ b/Engine/Source/Developer/TraceAnalysis/Public/Trace/Analyzer.h
@@ -61,6 +61,9 @@ public:
 
 		/** Is this field an array-type field? */
 		bool IsArray() const;
+
+		/** Is this field signed (only relevant for integer types) */
+		bool IsSigned() const;
 	};
 
 	struct FEventFieldHandle
@@ -107,7 +110,7 @@ public:
 		uint32 Num() const;
 
 	protected:
-		const void* GetImpl(uint32 Index, int16& SizeAndType) const;
+		const void* GetImpl(uint32 Index, int8& SizeAndType) const;
 	};
 
 	template <typename ValueType>
@@ -186,7 +189,7 @@ public:
 	private:
 		bool IsDefinitionImpl(uint32& OutTypeId) const;
 		const void* GetReferenceValueImpl(const char* FieldName, uint16& OutSizeType, uint32& OutTypeUid) const;
-		const void* GetValueImpl(const ANSICHAR* FieldName, int16& SizeAndType) const;
+		const void* GetValueImpl(const ANSICHAR* FieldName, int8& SizeAndType) const;
 		const FArrayReader* GetArrayImpl(const ANSICHAR* FieldName) const;
 	};
 
@@ -287,7 +290,7 @@ public:
 	}
 
 private:
-	template <typename ValueType> static ValueType CoerceValue(const void* Addr, int16 SizeAndType);
+	template <typename ValueType> static ValueType CoerceValue(const void* Addr, int8 SizeAndType);
 };
 
 ////////////////////////////////////////////////////////////////////////////////
@@ -309,26 +312,29 @@ IAnalyzer::FEventFieldHandle IAnalyzer::FEventTypeInfo::GetFieldHandle(const ANS
 
 ////////////////////////////////////////////////////////////////////////////////
 template <typename ValueType>
-ValueType IAnalyzer::CoerceValue(const void* Addr, int16 SizeAndType)
+ValueType IAnalyzer::CoerceValue(const void* Addr, int8 SizeAndType)
 {
+	using integral8 = std::conditional_t<std::is_signed_v<ValueType>, int8, uint8>;
+	using integral16 = std::conditional_t<std::is_signed_v<ValueType>, int16, uint16>;
+	using integral32 = std::conditional_t<std::is_signed_v<ValueType>, int32, uint32>;
+	using integral64 = std::conditional_t<std::is_signed_v<ValueType>, int64, uint64>;
 	switch (SizeAndType)
 	{
 	case -4: return ValueType(*(const float*)(Addr));
 	case -8: return ValueType(*(const double*)(Addr));
-	case  1: return ValueType(*(const uint8*)(Addr));
-	case  2: return ValueType(*(const uint16*)(Addr));
-	case  4: return ValueType(*(const uint32*)(Addr));
-	case  8: return ValueType(*(const uint64*)(Addr));
+	case  1: return ValueType(*(const integral8*)(Addr));
+	case  2: return ValueType(*(const integral16*)(Addr));
+	case  4: return ValueType(*(const integral32*)(Addr));
+	case  8: return ValueType(*(const integral64*)(Addr));
+	default: return ValueType(0);
 	}
-
-	return ValueType(0);
 }
 
 ////////////////////////////////////////////////////////////////////////////////
 template <typename ValueType>
 ValueType IAnalyzer::FEventData::GetValue(const ANSICHAR* FieldName, ValueType Default) const
 {
-	int16 FieldSizeAndType;
+	int8 FieldSizeAndType;
 	if (const void* Addr = GetValueImpl(FieldName, FieldSizeAndType))
 	{
 		return CoerceValue<ValueType>(Addr, FieldSizeAndType);
@@ -364,7 +370,7 @@ TArrayView<const ValueType> IAnalyzer::FEventData::GetArrayView(const ANSICHAR*
 template <typename ValueType>
 ValueType IAnalyzer::TArrayReader<ValueType>::operator [] (uint32 Index) const
 {
-	int16 ElementSizeAndType;
+	int8 ElementSizeAndType;
 	if (const void* Addr = GetImpl(Index, ElementSizeAndType))
 	{
 		return CoerceValue<ValueType>(Addr, ElementSizeAndType);
@@ -376,7 +382,7 @@ ValueType IAnalyzer::TArrayReader<ValueType>::operator [] (uint32 Index) const
 template <typename ValueType>
 const ValueType* IAnalyzer::TArrayReader<ValueType>::GetData() const
 {
-	int16 ElementSizeAndType;
+	int8 ElementSizeAndType;
 	const void* Addr = GetImpl(0, ElementSizeAndType);
 
 	if (Addr == nullptr || sizeof(ValueType) != abs(ElementSizeAndType))
diff --git a/Engine/Source/Runtime/TraceLog/Public/Trace/Detail/Field.h b/Engine/Source/Runtime/TraceLog/Public/Trace/Detail/Field.h
--- a/Engine/Source/Runtime/TraceLog/Public/Trace/Detail/Field.h
+++ b/Engine/Source/Runtime/TraceLog/Public/Trace/Detail/Field.h
@@ -43,10 +43,10 @@ template <> struct TFieldType<int8>			{ enum { Tid = int(EFieldType::Int8),	Size
 template <> struct TFieldType<int16>		{ enum { Tid = int(EFieldType::Int16),	Size = sizeof(int16) }; };
 template <> struct TFieldType<int32>		{ enum { Tid = int(EFieldType::Int32),	Size = sizeof(int32) }; };
 template <> struct TFieldType<int64>		{ enum { Tid = int(EFieldType::Int64),	Size = sizeof(int64) }; };
-template <> struct TFieldType<uint8>		{ enum { Tid = int(EFieldType::Int8),	Size = sizeof(uint8) }; };
-template <> struct TFieldType<uint16>		{ enum { Tid = int(EFieldType::Int16),	Size = sizeof(uint16) }; };
-template <> struct TFieldType<uint32>		{ enum { Tid = int(EFieldType::Int32),	Size = sizeof(uint32) }; };
-template <> struct TFieldType<uint64>		{ enum { Tid = int(EFieldType::Int64),	Size = sizeof(uint64) }; };
+template <> struct TFieldType<uint8>		{ enum { Tid = int(EFieldType::Uint8),	Size = sizeof(uint8) }; };
+template <> struct TFieldType<uint16>		{ enum { Tid = int(EFieldType::Uint16),	Size = sizeof(uint16) }; };
+template <> struct TFieldType<uint32>		{ enum { Tid = int(EFieldType::Uint32),	Size = sizeof(uint32) }; };
+template <> struct TFieldType<uint64>		{ enum { Tid = int(EFieldType::Uint64),	Size = sizeof(uint64) }; };
 template <> struct TFieldType<float>		{ enum { Tid = int(EFieldType::Float32),Size = sizeof(float) }; };
 template <> struct TFieldType<double>		{ enum { Tid = int(EFieldType::Float64),Size = sizeof(double) }; };
 template <class T> struct TFieldType<T*>	{ enum { Tid = int(EFieldType::Pointer),Size = sizeof(void*) }; };
diff --git a/Engine/Source/Runtime/TraceLog/Public/Trace/Detail/Protocols/Protocol0.h b/Engine/Source/Runtime/TraceLog/Public/Trace/Detail/Protocols/Protocol0.h
--- a/Engine/Source/Runtime/TraceLog/Public/Trace/Detail/Protocols/Protocol0.h
+++ b/Engine/Source/Runtime/TraceLog/Public/Trace/Detail/Protocols/Protocol0.h
@@ -39,7 +39,8 @@ enum : uint8
 	Field_SpecialMask	= 0030,
 	Field_Pod			= 0000,
 	Field_String		= 0010,
-	/*Field_Unused_2	= 0020,
+	Field_Signed		= 0020,
+	/*Field_Unused_3	= 0030,
 	  ...
 	  Field_Unused_7	= 0070,*/
 };
@@ -47,16 +48,20 @@ enum : uint8
 ////////////////////////////////////////////////////////////////////////////////
 enum class EFieldType : uint8
 {
-	Bool		= Field_Pod    | Field_Integer             | Field_8,
-	Int8		= Field_Pod    | Field_Integer             | Field_8,
-	Int16		= Field_Pod    | Field_Integer             | Field_16,
-	Int32		= Field_Pod    | Field_Integer             | Field_32,
-	Int64		= Field_Pod    | Field_Integer             | Field_64,
-	Pointer		= Field_Pod    | Field_Integer             | Field_Ptr,
-	Float32		= Field_Pod    | Field_Float               | Field_32,
-	Float64		= Field_Pod    | Field_Float               | Field_64,
-	AnsiString	= Field_String | Field_Integer|Field_Array | Field_8,
-	WideString	= Field_String | Field_Integer|Field_Array | Field_16,
+	Bool		= Field_Pod    | Field_Integer                 | Field_8,
+	Int8		= Field_Pod    | Field_Integer|Field_Signed    | Field_8,
+	Int16		= Field_Pod    | Field_Integer|Field_Signed    | Field_16,
+	Int32		= Field_Pod    | Field_Integer|Field_Signed    | Field_32,
+	Int64		= Field_Pod    | Field_Integer|Field_Signed    | Field_64,
+	Uint8		= Field_Pod    | Field_Integer                 | Field_8,
+	Uint16		= Field_Pod    | Field_Integer                 | Field_16,
+	Uint32		= Field_Pod    | Field_Integer                 | Field_32,
+	Uint64		= Field_Pod    | Field_Integer                 | Field_64,
+	Pointer		= Field_Pod    | Field_Integer                 | Field_Ptr,
+	Float32		= Field_Pod    | Field_Float                   | Field_32,
+	Float64		= Field_Pod    | Field_Float                   | Field_64,
+	AnsiString	= Field_String | Field_Integer|Field_Array     | Field_8,
+	WideString	= Field_String | Field_Integer|Field_Array     | Field_16,
 	Array		= Field_Array,
 };
 
