Expose raw event data in handler.

diff --git a/Engine/Source/Developer/TraceAnalysis/Private/Analysis/Engine.cpp b/Engine/Source/Developer/TraceAnalysis/Private/Analysis/Engine.cpp
--- a/Engine/Source/Developer/TraceAnalysis/Private/Analysis/Engine.cpp
+++ b/Engine/Source/Developer/TraceAnalysis/Private/Analysis/Engine.cpp
@@ -509,6 +509,17 @@ const IAnalyzer::FEventFieldInfo* IAnalyzer::FEventTypeInfo::GetFieldInfo(uint32
 	return (const IAnalyzer::FEventFieldInfo*)(Inner->Fields + Index);
 }
 
+////////////////////////////////////////////////////////////////////////////////
+IAnalyzer::FEventFieldHandle IAnalyzer::FEventTypeInfo::GetFieldHandleUnchecked(uint32 Index) const
+{
+	const auto* Inner = (const FDispatch*)this;
+	if (Index >= Inner->FieldCount)
+	{
+		return FEventFieldHandle { -1 };
+	}
+	return FEventFieldHandle {Inner->Fields[Index].Offset };
+}
+
 ////////////////////////////////////////////////////////////////////////////////
 IAnalyzer::FEventFieldHandle IAnalyzer::FEventTypeInfo::GetFieldHandleImpl(
 	const ANSICHAR* FieldName,
@@ -586,6 +597,12 @@ bool IAnalyzer::FEventFieldInfo::IsSigned() const
 	return (Inner->Class & UE::Trace::Protocol0::Field_Signed) != 0;
 }
 
+////////////////////////////////////////////////////////////////////////////////
+uint8 IAnalyzer::FEventFieldInfo::GetSize() const
+{
+	const auto* Inner = (const FDispatch::FField*)this;
+	return Inner->Size;
+}
 
 // {{{1 array-reader -----------------------------------------------------------
 
@@ -816,6 +833,17 @@ void IAnalyzer::FEventData::SerializeToCbor(TArray<uint8>& Out) const
 	SerializeToCborImpl(Out, *this, Size);
 }
 
+////////////////////////////////////////////////////////////////////////////////
+const void* IAnalyzer::FEventData::GetValueRaw(FEventFieldHandle Handle) const
+{
+	const auto* Info = (const FEventDataInfo*)this;
+	if (!Handle.IsValid())
+	{
+		return nullptr;
+	}
+	return Info->Ptr + Handle.Detail;
+}
+
 ////////////////////////////////////////////////////////////////////////////////
 const uint8* IAnalyzer::FEventData::GetAttachment() const
 {
@@ -1011,10 +1039,6 @@ const FTypeRegistry::FTypeInfo* FTypeRegistry::AddVersion6(const void* TraceData
 			OutField.bArray = !!(Field.Regular.TypeInfo & Protocol0::Field_Array);
 
 			NameCursor += Field.Regular.NameSize;
-			/*if (Field.FieldType == Protocol6::EFieldFamily::Definition)
-			{
-				Builder.SetIsDefinition();
-			}*/
 		}
 		else if (Field.FieldType == EFieldFamily::Reference)
 		{
@@ -1024,7 +1048,7 @@ const FTypeRegistry::FTypeInfo* FTypeRegistry::AddVersion6(const void* TraceData
 			auto& OutField = Builder.AddField(NameCursor, Field.Reference.NameSize, TypeSize);
 			OutField.Offset = Field.Reference.Offset;
 			OutField.SizeAndType = TypeSize;
-			OutField.RefUid = uint32(Field.Reference.RefUid) /*<< EKnownEventUids::_UidShift*/;
+			OutField.RefUid = uint32(Field.Reference.RefUid);
 			OutField.Class = 0;
 			OutField.bArray = false;
 
diff --git a/Engine/Source/Developer/TraceAnalysis/Public/Trace/Analyzer.h b/Engine/Source/Developer/TraceAnalysis/Public/Trace/Analyzer.h
--- a/Engine/Source/Developer/TraceAnalysis/Public/Trace/Analyzer.h
+++ b/Engine/Source/Developer/TraceAnalysis/Public/Trace/Analyzer.h
@@ -64,6 +64,10 @@ public:
 
 		/** Is this field signed (only relevant for integer types) */
 		bool IsSigned() const;
+
+		/** Gets the size in bytes for this field */
+		uint8 GetSize() const;
+
 	};
 
 	struct FEventFieldHandle
@@ -100,6 +104,12 @@ public:
 		template <typename ValueType>
 		FEventFieldHandle GetFieldHandle(const ANSICHAR* FieldName) const;
 
+		/** Returns a handle without specifying type. This should only be used in circumstances
+		 * where the field is treated untyped data.
+		 * @param Index Index of field.
+		 */
+		FEventFieldHandle GetFieldHandleUnchecked(uint32 Index) const;
+
 	private:
 		FEventFieldHandle GetFieldHandleImpl(const ANSICHAR*, int16&) const;
 	};
@@ -178,6 +188,13 @@ public:
 		 * @param Recipient of the Cbor serialization. Data is appeneded to Out. */
 		void SerializeToCbor(TArray<uint8>& Out) const;
 
+		/**
+		 * Returns the raw pointer to a field value
+		 * @param Handle Handle to field
+		 * @return Untyped pointer to field value
+		 */
+		const void* GetValueRaw(FEventFieldHandle Handle) const;
+
 		/** Returns the event's attachment. Not that this will always return an
 		 * address but if the event has no attachment then reading from that
 		 * address if undefined. */
