[Insights] Change cbor return type to UTF8StringView

Unreal Trace Server now writes full UTF8 strings to the CBOR objects, change the corresponding string views to `FUtf8StringView` to support correct conversion to TCHAR based strings.

diff --git a/Engine/Source/Developer/TraceAnalysis/Private/Store/AsioStoreCborServer.cpp b/Engine/Source/Developer/TraceAnalysis/Private/Store/AsioStoreCborServer.cpp
--- a/Engine/Source/Developer/TraceAnalysis/Private/Store/AsioStoreCborServer.cpp
+++ b/Engine/Source/Developer/TraceAnalysis/Private/Store/AsioStoreCborServer.cpp
@@ -218,8 +218,8 @@ void FAsioStoreCborPeer::OnTraceRead()
 ////////////////////////////////////////////////////////////////////////////////
 void FAsioStoreCborPeer::OnPayload()
 {
-	FAnsiStringView Request = Response.GetString("$request", "");
-	FAnsiStringView Path = Response.GetString("$path", "");
+	FUtf8StringView Request = Response.GetString("$request", "");
+	FUtf8StringView Path = Response.GetString("$path", "");
 	if (!Request.Len() || !Path.Len())
 	{
 		SendError(EStatusCode::BadRequest);
diff --git a/Engine/Source/Developer/TraceAnalysis/Private/Store/CborPayload.h b/Engine/Source/Developer/TraceAnalysis/Private/Store/CborPayload.h
--- a/Engine/Source/Developer/TraceAnalysis/Private/Store/CborPayload.h
+++ b/Engine/Source/Developer/TraceAnalysis/Private/Store/CborPayload.h
@@ -137,7 +137,7 @@ public:
 	EStatusCode		GetStatusCode() const;
 	int64			GetInteger(const char* Key, int64 Default) const;
 	template <int N>
-	FAnsiStringView	GetString(const char* Key, const char (&Default)[N]) const;
+	FUtf8StringView	GetString(const char* Key, const char (&Default)[N]) const;
 	const uint8*	GetData() const;
 	uint32			GetSize() const;
 	uint8*			Reserve(uint32 Size);
@@ -228,9 +228,9 @@ inline int64 FResponse::GetInteger(const char* Key, int64 Default) const
 
 ////////////////////////////////////////////////////////////////////////////////
 template <int N>
-inline FAnsiStringView FResponse::GetString(const char* Key, const char (&Default)[N]) const
+inline FUtf8StringView FResponse::GetString(const char* Key, const char (&Default)[N]) const
 {
-	FAnsiStringView DefaultView(Default, N - 1);
+	FUtf8StringView DefaultView(Default, N - 1);
 	return GetValue(
 		Key,
 		DefaultView,
@@ -240,7 +240,7 @@ inline FAnsiStringView FResponse::GetString(const char* Key, const char (&Defaul
 			{
 				int32 Length = Context.AsLength();
 				const char* Data = (const char*)(Buffer.GetData() + Offset - Length);
-				return FAnsiStringView(Data, Length);
+				return FUtf8StringView(Data, Length);
 			}
 
 			return DefaultView;
diff --git a/Engine/Source/Developer/TraceAnalysis/Private/Store/StoreClient.cpp b/Engine/Source/Developer/TraceAnalysis/Private/Store/StoreClient.cpp
--- a/Engine/Source/Developer/TraceAnalysis/Private/Store/StoreClient.cpp
+++ b/Engine/Source/Developer/TraceAnalysis/Private/Store/StoreClient.cpp
@@ -379,7 +379,7 @@ namespace UE {
 namespace Trace {
 
 ////////////////////////////////////////////////////////////////////////////////
-FAnsiStringView FStoreClient::FStatus::GetStoreDir() const
+FUtf8StringView FStoreClient::FStatus::GetStoreDir() const
 {
 	const auto* Response = (const FResponse*)this;
 	return Response->GetString("store_dir", "");
@@ -425,7 +425,7 @@ uint64 FStoreClient::FTraceInfo::GetTimestamp() const
 }
 
 ////////////////////////////////////////////////////////////////////////////////
-FAnsiStringView FStoreClient::FTraceInfo::GetName() const
+FUtf8StringView FStoreClient::FTraceInfo::GetName() const
 {
 	const auto* Response = (const FResponse*)this;
 	return Response->GetString("name", "nameless");
diff --git a/Engine/Source/Developer/TraceAnalysis/Public/Trace/StoreClient.h b/Engine/Source/Developer/TraceAnalysis/Public/Trace/StoreClient.h
--- a/Engine/Source/Developer/TraceAnalysis/Public/Trace/StoreClient.h
+++ b/Engine/Source/Developer/TraceAnalysis/Public/Trace/StoreClient.h
@@ -18,14 +18,14 @@ class TRACEANALYSIS_API FStoreClient
 public:
 	struct TRACEANALYSIS_API FStatus
 	{
-		FAnsiStringView	GetStoreDir() const;
+		FUtf8StringView	GetStoreDir() const;
 		uint32			GetRecorderPort() const;
 		uint32			GetChangeSerial() const;
 	};
 
 	struct TRACEANALYSIS_API FTraceInfo
 	{
-		FAnsiStringView	GetName() const;
+		FUtf8StringView	GetName() const;
 		uint32			GetId() const;
 		uint64			GetSize() const;
 		uint64			GetTimestamp() const;
diff --git a/Engine/Source/Developer/TraceInsights/Private/Insights/InsightsManager.cpp b/Engine/Source/Developer/TraceInsights/Private/Insights/InsightsManager.cpp
--- a/Engine/Source/Developer/TraceInsights/Private/Insights/InsightsManager.cpp
+++ b/Engine/Source/Developer/TraceInsights/Private/Insights/InsightsManager.cpp
@@ -877,7 +877,7 @@ void FInsightsManager::LoadTrace(uint32 InTraceId, bool InAutoQuit)
 	const UE::Trace::FStoreClient::FTraceInfo* TraceInfo = StoreClient->GetTraceInfoById(InTraceId);
 	if (TraceInfo != nullptr)
 	{
-		FAnsiStringView Name = TraceInfo->GetName();
+		FUtf8StringView Name = TraceInfo->GetName();
 		TraceName = FPaths::Combine(TraceName, FString(Name.Len(), Name.GetData()));
 		TraceName = FPaths::SetExtension(TraceName, ".utrace");
 	}
diff --git a/Engine/Source/Developer/TraceInsights/Private/Insights/StoreService/StoreBrowser.cpp b/Engine/Source/Developer/TraceInsights/Private/Insights/StoreService/StoreBrowser.cpp
--- a/Engine/Source/Developer/TraceInsights/Private/Insights/StoreService/StoreBrowser.cpp
+++ b/Engine/Source/Developer/TraceInsights/Private/Insights/StoreService/StoreBrowser.cpp
@@ -164,8 +164,8 @@ void FStoreBrowser::UpdateTraces()
 							Trace.TraceId = TraceId;
 							Trace.TraceIndex = TraceIndex;
 
-							const FAnsiStringView AnsiTraceName = TraceInfo->GetName();
-							Trace.Name = FString(AnsiTraceName.Len(), AnsiTraceName.GetData());
+							const FUtf8StringView Utf8NameView = TraceInfo->GetName();
+							Trace.Name = FString(Utf8NameView);
 
 							Trace.Timestamp = FStoreBrowserTraceInfo::ConvertTimestamp(TraceInfo->GetTimestamp());
 							Trace.Size = TraceInfo->GetSize();
