Fix serial gaps issue when no gaps are filled, but new appear, during the second sync phase.

#rnx

diff --git a/Engine/Source/Developer/TraceAnalysis/Private/Analysis/Engine.cpp b/Engine/Source/Developer/TraceAnalysis/Private/Analysis/Engine.cpp
--- a/Engine/Source/Developer/TraceAnalysis/Private/Analysis/Engine.cpp
+++ b/Engine/Source/Developer/TraceAnalysis/Private/Analysis/Engine.cpp
@@ -3512,7 +3512,7 @@ void FProtocol5Stage::DetectSerialGaps(TArray<FEventDescStream>& EventDescHeap)
 		uint32 GapCount = 0;
 		auto RecordGap = [this, &GapCount] (int32 Lhs, int32 Rhs)
 		{
-			if (SerialGaps.IsEmpty())
+			if (SerialGaps.IsEmpty() || GapCount >= (uint32)SerialGaps.Num())
 			{
 				return false;
 			}
