[Insights] Fixed analysis to continue reading at least one more loop after encountering a SYNC point.

diff --git a/Engine/Source/Developer/TraceAnalysis/Private/Analysis/Engine.cpp b/Engine/Source/Developer/TraceAnalysis/Private/Analysis/Engine.cpp
--- a/Engine/Source/Developer/TraceAnalysis/Private/Analysis/Engine.cpp
+++ b/Engine/Source/Developer/TraceAnalysis/Private/Analysis/Engine.cpp
@@ -1872,6 +1872,7 @@ public:
 		NotEnoughData,
 		EndOfStream,
 		Continue,
+		Sync,
 	};
 
 	struct FMachineContext
@@ -2826,15 +2827,13 @@ FProtocol5Stage::EStatus FProtocol5Stage::OnData(
 		return Ret;
 	}
 
-	bool bNotEnoughData;
-
 	// Important events
 	Ret = OnDataImportant(Context);
 	if (Ret == EStatus::Error)
 	{
 		return Ret;
 	}
-	bNotEnoughData = (Ret != EStatus::EndOfStream);
+	bool bNotEnoughData = (Ret == EStatus::NotEnoughData);
 
 	// Normal events
 	Ret = OnDataNormal(Context);
@@ -2842,6 +2841,11 @@ FProtocol5Stage::EStatus FProtocol5Stage::OnData(
 	{
 		return Ret;
 	}
+	if (Ret == EStatus::Sync)
+	{
+		// After processing a SYNC packet, we need to read data once more.
+		return OnData(Reader, Context);
+	}
 	bNotEnoughData |= (Ret == EStatus::NotEnoughData);
 
 	if (bNotEnoughData)
@@ -3097,6 +3101,7 @@ FProtocol5Stage::EStatus FProtocol5Stage::OnDataNormal(const FMachineContext& Co
 		NextSerial = EventDescHeap.HeapTop().EventDescs[0].Serial;
 	}
 
+	const bool bSync = (SyncCount != Transport.GetSyncCount());
 	DetectSerialGaps(EventDescHeap);
 
 	if (DispatchEvents(Context.Bridge, EventDescHeap) < 0)
@@ -3104,6 +3109,11 @@ FProtocol5Stage::EStatus FProtocol5Stage::OnDataNormal(const FMachineContext& Co
 		return EStatus::Error;
 	}
 
+	if (bSync)
+	{
+		return EStatus::Sync;
+	}
+
 	return bNotEnoughData ? EStatus::NotEnoughData : EStatus::EndOfStream;
 }
 
