[Insights] Fixed serialization to Cbor for EventData with zero fields.

diff --git a/Engine/Source/Developer/TraceAnalysis/Private/Analysis/EventToCbor.cpp b/Engine/Source/Developer/TraceAnalysis/Private/Analysis/EventToCbor.cpp
--- a/Engine/Source/Developer/TraceAnalysis/Private/Analysis/EventToCbor.cpp
+++ b/Engine/Source/Developer/TraceAnalysis/Private/Analysis/EventToCbor.cpp
@@ -37,19 +37,24 @@ void SerializeToCborImpl(
 	 * of the analysis engine instead */
 
 	const IAnalyzer::FEventTypeInfo& TypeInfo = EventData.GetTypeInfo();
+	const uint32 FieldCount = TypeInfo.GetFieldCount();
+	if (FieldCount == 0)
+	{
+		return;
+	}
 
 	uint32 SizeHint = ((EventSize * 11) / 10); // + 10%
-	SizeHint += TypeInfo.GetFieldCount() * 16;
+	SizeHint += FieldCount * 16;
 	SizeHint += Out.Num();
 
 	Out.Reserve(SizeHint);
 	FMemoryWriter MemoryWriter(Out, false, true);
 	FCborWriter CborWriter(&MemoryWriter, ECborEndianness::StandardCompliant);
 
-	CborWriter.WriteContainerStart(ECborCode::Map, TypeInfo.GetFieldCount());
-	for (uint32 i = 0, n = TypeInfo.GetFieldCount(); i < n; ++i)
+	CborWriter.WriteContainerStart(ECborCode::Map, FieldCount);
+	for (uint32 FieldIndex = 0; FieldIndex < FieldCount; ++FieldIndex)
 	{
-		const IAnalyzer::FEventFieldInfo& FieldInfo = *(TypeInfo.GetFieldInfo(i));
+		const IAnalyzer::FEventFieldInfo& FieldInfo = *(TypeInfo.GetFieldInfo(FieldIndex));
 
 		const ANSICHAR* FieldName = FieldInfo.GetName();
 		CborWriter.WriteValue(FieldName, FCStringAnsi::Strlen(FieldName));
