[Insights] Fixed unsafe type cast warnings in TraceAnalysis module.

diff --git a/Engine/Source/Developer/TraceAnalysis/Private/Analysis/Engine.cpp b/Engine/Source/Developer/TraceAnalysis/Private/Analysis/Engine.cpp
--- a/Engine/Source/Developer/TraceAnalysis/Private/Analysis/Engine.cpp
+++ b/Engine/Source/Developer/TraceAnalysis/Private/Analysis/Engine.cpp
@@ -448,14 +448,14 @@ FDispatch* FDispatchBuilder::Finalize()
 	for (int i = 0, n = Dispatch->FieldCount; i < n; ++i)
 	{
 		auto* Field = Dispatch->Fields + i;
-		Field->NameOffset += Buffer.Num() - uint32(UPTRINT(Field) - UPTRINT(Dispatch));
+		Field->NameOffset += (uint16)(Buffer.Num() - uint32(UPTRINT(Field) - UPTRINT(Dispatch)));
 	}
 
 	// Calculate this dispatch's hash.
 	if (Dispatch->LoggerNameOffset || Dispatch->EventNameOffset)
 	{
-		Dispatch->LoggerNameOffset += Buffer.Num();
-		Dispatch->EventNameOffset += Buffer.Num();
+		Dispatch->LoggerNameOffset += (uint16)Buffer.Num();
+		Dispatch->EventNameOffset += (uint16)Buffer.Num();
 
 		FFnv1aHash Hash;
 		Hash.Add((const ANSICHAR*)Dispatch + Dispatch->LoggerNameOffset);
@@ -477,14 +477,14 @@ void FDispatchBuilder::SetUid(uint16 Uid)
 void FDispatchBuilder::SetLoggerName(const ANSICHAR* Name, int32 NameSize)
 {
 	auto* Dispatch = (FDispatch*)(Buffer.GetData());
-	Dispatch->LoggerNameOffset += AppendName(Name, NameSize);
+	Dispatch->LoggerNameOffset += (uint16)AppendName(Name, NameSize);
 }
 
 ////////////////////////////////////////////////////////////////////////////////
 void FDispatchBuilder::SetEventName(const ANSICHAR* Name, int32 NameSize)
 {
 	auto* Dispatch = (FDispatch*)(Buffer.GetData());
-	Dispatch->EventNameOffset = AppendName(Name, NameSize);
+	Dispatch->EventNameOffset = (uint16)AppendName(Name, NameSize);
 }
 
 ////////////////////////////////////////////////////////////////////////////////
@@ -520,7 +520,7 @@ FDispatch::FField& FDispatchBuilder::AddField(const ANSICHAR* Name, int32 NameSi
 {
 	int32 Bufoff = Buffer.AddUninitialized(sizeof(FDispatch::FField));
 	auto* Field = (FDispatch::FField*)(Buffer.GetData() + Bufoff);
-	Field->NameOffset = AppendName(Name, NameSize);
+	Field->NameOffset = (uint16)AppendName(Name, NameSize);
 	Field->Size = Size;
 	Field->RefUid = 0;
 
@@ -693,7 +693,7 @@ bool IAnalyzer::FEventFieldInfo::IsSigned() const
 uint8 IAnalyzer::FEventFieldInfo::GetSize() const
 {
 	const auto* Inner = (const FDispatch::FField*)this;
-	return Inner->Size;
+	return uint8(Inner->Size);
 }
 
 // {{{1 array-reader -----------------------------------------------------------
@@ -1076,7 +1076,7 @@ const FTypeRegistry::FTypeInfo* FTypeRegistry::AddVersion4(const void* TraceData
 	{
 		const auto& Field = NewEvent.Fields[i];
 
-		uint16 TypeSize = 1 << (Field.TypeInfo & Protocol0::Field_Pow2SizeMask);
+		int8 TypeSize = 1 << (Field.TypeInfo & Protocol0::Field_Pow2SizeMask);
 		if (Field.TypeInfo & Protocol0::Field_Float)
 		{
 			TypeSize = -TypeSize;
@@ -1136,7 +1136,7 @@ const FTypeRegistry::FTypeInfo* FTypeRegistry::AddVersion6(const void* TraceData
 		const auto& Field = NewEvent.Fields[i];
 		if (Field.FieldType == EFieldFamily::Regular)
 		{
-			uint16 TypeSize = 1 << (Field.Regular.TypeInfo & Protocol0::Field_Pow2SizeMask);
+			int8 TypeSize = 1 << (Field.Regular.TypeInfo & Protocol0::Field_Pow2SizeMask);
 			if (Field.Regular.TypeInfo & Protocol0::Field_Float)
 			{
 				TypeSize = -TypeSize;
@@ -1154,7 +1154,7 @@ const FTypeRegistry::FTypeInfo* FTypeRegistry::AddVersion6(const void* TraceData
 		else if (Field.FieldType == EFieldFamily::Reference)
 		{
 			check((Field.Reference.TypeInfo & Protocol0::Field_CategoryMask) == Protocol0::Field_Integer);
-			const uint16 TypeSize = 1 << (Field.Reference.TypeInfo & Protocol0::Field_Pow2SizeMask);
+			const int8 TypeSize = 1 << (Field.Reference.TypeInfo & Protocol0::Field_Pow2SizeMask);
 
 			auto& OutField = Builder.AddField(NameCursor, Field.Reference.NameSize, TypeSize);
 			OutField.Offset = Field.Reference.Offset;
@@ -1168,7 +1168,7 @@ const FTypeRegistry::FTypeInfo* FTypeRegistry::AddVersion6(const void* TraceData
 		else if (Field.FieldType == EFieldFamily::DefinitionId)
 		{
 			check((Field.DefinitionId.TypeInfo & Protocol0::Field_CategoryMask) == Protocol0::Field_Integer);
-			const uint16 TypeSize = 1 << (Field.DefinitionId.TypeInfo & Protocol0::Field_Pow2SizeMask);
+			const int8 TypeSize = 1 << (Field.DefinitionId.TypeInfo & Protocol0::Field_Pow2SizeMask);
 			
 			auto DefinitionIdFieldName = ANSITEXTVIEW("DefinitionId");
 			auto& OutField = Builder.AddField(DefinitionIdFieldName.GetData(), DefinitionIdFieldName.Len(), TypeSize);
@@ -1460,7 +1460,7 @@ void FAnalyzerHub::BuildRoutes()
 	Builder.Self = this;
 
 	IAnalyzer::FOnAnalysisContext OnAnalysisContext = { Builder };
-	for (uint16 i = 0, n = Analyzers.Num(); i < n; ++i)
+	for (uint16 i = 0, n = uint16(Analyzers.Num()); i < n; ++i)
 	{
 		uint32 RouteCount = Routes.Num();
 
@@ -1502,7 +1502,7 @@ void FAnalyzerHub::BuildRoutes()
 	};
 
 	FRoute* Cursor = FixupRoute(Routes.GetData());
-	for (uint16 i = 1, n = Routes.Num(); i < n; ++i)
+	for (uint16 i = 1, n = uint16(Routes.Num()); i < n; ++i)
 	{
 		FRoute* Route = Routes.GetData() + i;
 		if (Route->Hash == Cursor->Hash)
diff --git a/Engine/Source/Developer/TraceAnalysis/Private/Analysis/Transport/PacketTransport.cpp b/Engine/Source/Developer/TraceAnalysis/Private/Analysis/Transport/PacketTransport.cpp
--- a/Engine/Source/Developer/TraceAnalysis/Private/Analysis/Transport/PacketTransport.cpp
+++ b/Engine/Source/Developer/TraceAnalysis/Private/Analysis/Transport/PacketTransport.cpp
@@ -153,7 +153,7 @@ bool FPacketTransport::GetNextBatch()
 			};
 			auto* PacketEncoded = (FPacketEncoded*)PacketBase;
 
-			Node->Size = UE::Trace::Private::Decode(
+			Node->Size = (uint16)UE::Trace::Private::Decode(
 				PacketEncoded->Data,
 				int32(PacketEncoded->PacketSize - sizeof(FPacketEncoded)),
 				Node->Data,
diff --git a/Engine/Source/Developer/TraceAnalysis/Private/DataStream.cpp b/Engine/Source/Developer/TraceAnalysis/Private/DataStream.cpp
--- a/Engine/Source/Developer/TraceAnalysis/Private/DataStream.cpp
+++ b/Engine/Source/Developer/TraceAnalysis/Private/DataStream.cpp
@@ -40,15 +40,19 @@ int32 FFileDataStream::Read(void* Data, uint32 Size)
 		return 0;
 	}
 
-	Size = (Size < Remaining) ? Size : Remaining;
+	if (Size > Remaining)
+	{
+		Size = static_cast<uint32>(Remaining);
+	}
 	Remaining -= Size;
-	int32 Result = Handle->Read((uint8*)Data, Size) ? Size : -1;
-	if (Result < 0)
+
+	if (!Handle->Read((uint8*)Data, Size))
 	{
 		Close();
+		return -1;
 	}
 
-	return Result;
+	return Size;
 }
 
 void FFileDataStream::Close()
diff --git a/Engine/Source/Developer/TraceAnalysis/Private/Store/CborPayload.h b/Engine/Source/Developer/TraceAnalysis/Private/Store/CborPayload.h
--- a/Engine/Source/Developer/TraceAnalysis/Private/Store/CborPayload.h
+++ b/Engine/Source/Developer/TraceAnalysis/Private/Store/CborPayload.h
@@ -38,7 +38,7 @@ inline TInlineMemoryWriter<BufferSize>::TInlineMemoryWriter()
 template <int BufferSize>
 inline void TInlineMemoryWriter<BufferSize>::Serialize(void* Data, int64 Num)
 {
-	Buffer.Append((const uint8*)Data, Num);
+	Buffer.Append((const uint8*)Data, IntCastChecked<int32>(Num));
 }
 
 
@@ -135,7 +135,11 @@ class FResponse
 {
 public:
 	EStatusCode		GetStatusCode() const;
-	int64			GetInteger(const char* Key, int64 Default) const;
+	int64			GetInt64Checked(const char* Key, int64 Default) const;
+	uint64			GetUint64Checked(const char* Key, uint64 Default) const;
+	int32			GetInt32Checked(const char* Key, int32 Default) const;
+	uint32			GetUint32Checked(const char* Key, uint32 Default) const;
+	int64			GetInteger(const char* Key, int64 Default) const { return GetInt64Checked(Key, Default); }
 	template <int N>
 	FUtf8StringView	GetString(const char* Key, const char (&Default)[N]) const;
 	const uint8*	GetData() const;
@@ -151,7 +155,7 @@ private:
 ////////////////////////////////////////////////////////////////////////////////
 inline EStatusCode FResponse::GetStatusCode() const
 {
-	int32 Code = GetInteger("$status", 0);
+	int32 Code = GetInt32Checked("$status", 0);
 	return Code ? EStatusCode(Code) : EStatusCode::Unknown;
 }
 
@@ -195,8 +199,8 @@ inline Type FResponse::GetValue(const char* Key, Type Default, LambdaType&& Lamb
 			return Default;
 		}
 
-		uint32 Length = Context.AsLength();
-		uint32 Offset = uint32(MemoryReader.Tell());
+		uint32 Length = static_cast<uint32>(Context.AsLength());
+		uint32 Offset = static_cast<uint32>(MemoryReader.Tell());
 		auto* String = (const char*)(Buffer.GetData() + Offset - Length);
 		bool bIsTarget = (FCStringAnsi::Strncmp(Key, String, Length) == 0);
 
@@ -208,13 +212,13 @@ inline Type FResponse::GetValue(const char* Key, Type Default, LambdaType&& Lamb
 
 		if (bIsTarget)
 		{
-			return Lambda(Context, uint32(MemoryReader.Tell()));
+			return Lambda(Context, static_cast<uint32>(MemoryReader.Tell()));
 		}
 	}
 }
 
 ////////////////////////////////////////////////////////////////////////////////
-inline int64 FResponse::GetInteger(const char* Key, int64 Default) const
+inline int64 FResponse::GetInt64Checked(const char* Key, int64 Default) const
 {
 	return GetValue(
 		Key,
@@ -226,6 +230,49 @@ inline int64 FResponse::GetInteger(const char* Key, int64 Default) const
 	);
 }
 
+////////////////////////////////////////////////////////////////////////////////
+inline uint64 FResponse::GetUint64Checked(const char* Key, uint64 Default) const
+{
+	return GetValue(
+		Key,
+		Default,
+		[this](const FCborContext& Context, uint64)
+		{
+			return Context.AsUInt();
+		}
+	);
+}
+
+////////////////////////////////////////////////////////////////////////////////
+inline int32 FResponse::GetInt32Checked(const char* Key, int32 Default) const
+{
+	return GetValue(
+		Key,
+		Default,
+		[this](const FCborContext& Context, int32)
+		{
+			int64 Value = Context.AsInt();
+			check(Value >= INT_MIN && Value <= INT_MAX);
+			return static_cast<int32>(Value);
+		}
+	);
+}
+
+////////////////////////////////////////////////////////////////////////////////
+inline uint32 FResponse::GetUint32Checked(const char* Key, uint32 Default) const
+{
+	return GetValue(
+		Key,
+		Default,
+		[this](const FCborContext& Context, uint32)
+		{
+			uint64 Value = Context.AsUInt();
+			check(Value <= UINT_MAX);
+			return static_cast<uint32>(Value);
+		}
+	);
+}
+
 ////////////////////////////////////////////////////////////////////////////////
 template <int N>
 inline FUtf8StringView FResponse::GetString(const char* Key, const char (&Default)[N]) const
@@ -238,7 +285,7 @@ inline FUtf8StringView FResponse::GetString(const char* Key, const char (&Defaul
 		{
 			if (Context.IsString())
 			{
-				int32 Length = Context.AsLength();
+				uint32 Length = static_cast<uint32>(Context.AsLength());
 				const char* Data = (const char*)(Buffer.GetData() + Offset - Length);
 				return FUtf8StringView(Data, Length);
 			}
diff --git a/Engine/Source/Developer/TraceAnalysis/Private/Store/StoreClient.cpp b/Engine/Source/Developer/TraceAnalysis/Private/Store/StoreClient.cpp
--- a/Engine/Source/Developer/TraceAnalysis/Private/Store/StoreClient.cpp
+++ b/Engine/Source/Developer/TraceAnalysis/Private/Store/StoreClient.cpp
@@ -67,7 +67,7 @@ int32 FTraceDataStream::Read(void* Dest, uint32 DestSize)
 	while (true)
 	{
 		fd_set ReadFds = Fds;
-		int Ret = select(Handle + 1, &ReadFds, 0, 0, &Timeout);
+		int Ret = select((int)Handle + 1, &ReadFds, 0, 0, &Timeout);
 
 		if (Ret < 0)
 		{
@@ -308,7 +308,7 @@ FTraceDataStream* FStoreCborClient::ReadTrace(uint32 Id)
 		return nullptr;
 	}
 
-	uint32 SenderPort = Response.GetInteger("port", 0);
+	uint32 SenderPort = Response.GetUint32Checked("port", 0);
 	if (!SenderPort)
 	{
 		return nullptr;
@@ -389,14 +389,14 @@ FUtf8StringView FStoreClient::FStatus::GetStoreDir() const
 uint32 FStoreClient::FStatus::GetRecorderPort() const
 {
 	const auto* Response = (const FResponse*)this;
-	return Response->GetInteger("recorder_port", 0);
+	return Response->GetUint32Checked("recorder_port", 0);
 }
 
 ////////////////////////////////////////////////////////////////////////////////
 uint32 FStoreClient::FStatus::GetChangeSerial() const
 {
 	const auto* Response = (const FResponse*)this;
-	return Response->GetInteger("change_serial", 0);
+	return Response->GetUint32Checked("change_serial", 0);
 }
 
 
@@ -407,21 +407,21 @@ uint32 FStoreClient::FStatus::GetChangeSerial() const
 uint32 FStoreClient::FTraceInfo::GetId() const
 {
 	const auto* Response = (const FResponse*)this;
-	return Response->GetInteger("id", 0);
+	return Response->GetUint32Checked("id", 0);
 }
 
 ////////////////////////////////////////////////////////////////////////////////
 uint64 FStoreClient::FTraceInfo::GetSize() const
 {
 	const auto* Response = (const FResponse*)this;
-	return Response->GetInteger("size", 0);
+	return Response->GetUint64Checked("size", 0);
 }
 
 ////////////////////////////////////////////////////////////////////////////////
 uint64 FStoreClient::FTraceInfo::GetTimestamp() const
 {
 	const auto* Response = (const FResponse*)this;
-	return Response->GetInteger("timestamp", 0);
+	return Response->GetUint64Checked("timestamp", 0);
 }
 
 ////////////////////////////////////////////////////////////////////////////////
@@ -437,28 +437,28 @@ FUtf8StringView FStoreClient::FTraceInfo::GetName() const
 uint32 FStoreClient::FSessionInfo::GetId() const
 {
 	const auto* Response = (const FResponse*)this;
-	return Response->GetInteger("id", 0);
+	return Response->GetUint32Checked("id", 0);
 }
 
 ////////////////////////////////////////////////////////////////////////////////
 uint32 FStoreClient::FSessionInfo::GetTraceId() const
 {
 	const auto* Response = (const FResponse*)this;
-	return Response->GetInteger("trace_id", 0);
+	return Response->GetUint32Checked("trace_id", 0);
 }
 
 ////////////////////////////////////////////////////////////////////////////////
 uint32 FStoreClient::FSessionInfo::GetIpAddress() const
 {
 	const auto* Response = (const FResponse*)this;
-	return Response->GetInteger("ip_address", 0);
+	return Response->GetUint32Checked("ip_address", 0);
 }
 
 ////////////////////////////////////////////////////////////////////////////////
 uint32 FStoreClient::FSessionInfo::GetControlPort() const
 {
 	const auto* Response = (const FResponse*)this;
-	return Response->GetInteger("control_port", 0);
+	return Response->GetUint32Checked("control_port", 0);
 }
 
 
@@ -469,7 +469,7 @@ FStoreClient* FStoreClient::Connect(const TCHAR* Host, uint32 Port)
 	static asio::io_context IoContext;
 
 	FStoreCborClient* Impl = new FStoreCborClient(IoContext);
-	if (!Impl->Connect(Host, Port))
+	if (!Impl->Connect(Host, static_cast<uint16>(Port)))
 	{
 		delete Impl;
 		return nullptr;
@@ -528,7 +528,7 @@ uint32 FStoreClient::GetTraceCount()
 		return 0;
 	}
 
-	return Self->GetResponse().GetInteger("count", 0);
+	return Self->GetResponse().GetUint32Checked("count", 0);
 }
 
 ////////////////////////////////////////////////////////////////////////////////
@@ -573,7 +573,7 @@ uint32 FStoreClient::GetSessionCount() const
 		return 0;
 	}
 
-	return Self->GetResponse().GetInteger("count", 0);
+	return Self->GetResponse().GetUint32Checked("count", 0);
 }
 
 ////////////////////////////////////////////////////////////////////////////////
