Apply UnrealTraceServer@23233583,23246589 to POSIX platforms.

#rnx

diff --git a/Engine/Source/Developer/TraceAnalysis/Private/Store/AsioRecorder.cpp b/Engine/Source/Developer/TraceAnalysis/Private/Store/AsioRecorder.cpp
--- a/Engine/Source/Developer/TraceAnalysis/Private/Store/AsioRecorder.cpp
+++ b/Engine/Source/Developer/TraceAnalysis/Private/Store/AsioRecorder.cpp
@@ -21,39 +21,89 @@ class FAsioRecorderRelay
 	: public FAsioIoSink
 {
 public:
-						FAsioRecorderRelay(asio::ip::tcp::socket& Socket, FAsioWriteable* InOutput);
+						FAsioRecorderRelay(asio::ip::tcp::socket& Socket, FAsioStore& InStore);
 	virtual				~FAsioRecorderRelay();
 	bool				IsOpen();
 	void				Close();
+	uint32				GetTraceId() const;
 	uint32				GetIpAddress() const;
 	uint32				GetControlPort() const;
 
 private:
 	virtual void		OnIoComplete(uint32 Id, int32 Size) override;
+	bool				CreateTrace();
+	bool				ReadMagic();
 	bool				ReadMetadata(int32 Size);
 	static const uint32	BufferSize = 64 * 1024;
-	enum				{ OpStart, OpSocketReadMetadata, OpSocketRead, OpFileWrite };
 	FAsioSocket			Input;
-	FAsioWriteable*		Output;
-	uint32				ActiveReadOp = OpSocketReadMetadata;
+	FAsioWriteable*		Output = nullptr;
+	FAsioStore&			Store;
+	uint8*				PreambleCursor;
+	uint32				TraceId = 0;
 	uint16				ControlPort = 0;
 	uint8				Buffer[BufferSize];
+
+	enum
+	{
+		OpMagicRead,
+		OpMetadataRead,
+		OpSocketRead,
+		OpFileWrite,
+	};
+
+	using MagicType				= uint32;
+	using MetadataSizeType		= uint16;
+	using VersionType			= struct { uint8 Transport; uint8 Protocol; };
+
+	static_assert(sizeof(VersionType) == 2, "Unexpected struct size");
 };
 
 ////////////////////////////////////////////////////////////////////////////////
-FAsioRecorderRelay::FAsioRecorderRelay(asio::ip::tcp::socket& Socket, FAsioWriteable* InOutput)
+FAsioRecorderRelay::FAsioRecorderRelay(asio::ip::tcp::socket& Socket, FAsioStore& InStore)
 : Input(Socket)
-, Output(InOutput)
+, Store(InStore)
 {
-	OnIoComplete(OpStart, 0);
+#if PLATFORM_WINDOWS
+	// Trace data is a stream and communication is one way. It is implemented
+	// this way to share code between sending trace data over the wire and writing
+	// it to a file. Because there's no ping/pong we can end up with a half-open
+	// TCP connection if the other end doesn't close its socket. So we'll enable
+	// keep-alive on the socket and set a short timeout (default is 2hrs).
+	tcp_keepalive KeepAlive =
+	{
+		1,		// on
+		15000,	// timeout_ms
+		2000,	// interval_ms
+	};
+
+	DWORD BytesReturned;
+	WSAIoctl(
+		Socket.native_handle(),
+		SIO_KEEPALIVE_VALS,
+		&KeepAlive, sizeof(KeepAlive),
+		nullptr, 0,
+		&BytesReturned,
+		nullptr,
+		nullptr
+	);
+#endif
+
+	// Kick things off by reading the magic four bytes at the start of the stream
+	// along with an additional two bytes that are likely the metadata size.
+	uint32 PreambleReadSize = sizeof(MagicType) + sizeof(MetadataSizeType);
+	PreambleCursor = Buffer + PreambleReadSize;
+	Input.Read(Buffer, PreambleReadSize, this, OpMagicRead);
 }
 
 ////////////////////////////////////////////////////////////////////////////////
 FAsioRecorderRelay::~FAsioRecorderRelay()
 {
 	check(!Input.IsOpen());
-	check(!Output->IsOpen());
-	delete Output;
+	if (Output != nullptr)
+	{
+		check(!Output->IsOpen());
+		delete Output;
+	}
 }
 
 ////////////////////////////////////////////////////////////////////////////////
@@ -66,7 +116,16 @@ bool FAsioRecorderRelay::IsOpen()
 void FAsioRecorderRelay::Close()
 {
 	Input.Close();
-	Output->Close();
+	if (Output != nullptr)
+	{
+		Output->Close();
+	}
+}
+
+////////////////////////////////////////////////////////////////////////////////
+uint32 FAsioRecorderRelay::GetTraceId() const
+{
+	return TraceId;
 }
 
 ////////////////////////////////////////////////////////////////////////////////
@@ -82,48 +141,81 @@ uint32 FAsioRecorderRelay::GetControlPort() const
 }
 
 ////////////////////////////////////////////////////////////////////////////////
-bool FAsioRecorderRelay::ReadMetadata(int32 Size)
+bool FAsioRecorderRelay::CreateTrace()
 {
-	const uint8* Cursor = Buffer;
-	auto Read = [&] (int32 SizeToRead)
-	{
-		const uint8* Ptr = Cursor;
-		Cursor += SizeToRead;
-		Size -= SizeToRead;
-		return Ptr;
-	};
+	FAsioStore::FNewTrace Trace = Store.CreateTrace();
+	TraceId = Trace.Id;
+	Output = Trace.Writeable;
+	return (Output != nullptr);
+}
+
+////////////////////////////////////////////////////////////////////////////////
+bool FAsioRecorderRelay::ReadMagic()
+{
+	// Here we'll check the magic four bytes at the start of the stream and create
+	// a trace to write into if they are bytes we're expecting.
 
-	// Stream header
-	uint32 Magic;
-	if (Size < sizeof(Magic))
+	// We will only support clients that send the magic. Very early clients did
+	// not do this but they were unreleased and should no longer be in use.
+	if (Buffer[3] != 'T' || Buffer[2] != 'R' || Buffer[1] != 'C')
 	{
-		return true;
+		return false;
 	}
 
-	Magic = *(const uint32*)(Read(sizeof(Magic)));
-	switch (Magic)
+	// We can continue to support very old clients
+	if (Buffer[0] == 'E')
 	{
-	case 'TRC2':		/* trace with metadata data */
-		break;
-
-	case 'TRCE':		/* valid, but to old or wrong endian for us */
-	case 'ECRT':
-	case '2CRT':
-		return true;
+		if (CreateTrace())
+		{
+			// Old clients have no metadata so we can go straight into the
+			// read-write loop. We've already got read data in Buffer.
+			Output->Write(Buffer, sizeof(MagicType), this, OpFileWrite);
+			return true;
+		}
+		return false;
+	}
 
-	default:			/* unexpected magic */
+	// Later clients have a metadata block (TRC2). There's loose support for the
+	// future too if need be (TRC[3-9]).
+	if (Buffer[0] < '2' || Buffer[0] > '9')
+	{
 		return false;
 	}
 
-	// MetadataSize field
-	if (Size < 2)
+	// Concatenate metadata into the buffer, first validating the given size is
+	// one that we can handle in a single read.
+	uint32 MetadataSize = *(MetadataSizeType*)(Buffer + sizeof(MagicType));
+	MetadataSize += sizeof(VersionType);
+	if (MetadataSize > BufferSize - uint32(ptrdiff_t(PreambleCursor - Buffer)))
 	{
-		return true;
+		return false;
 	}
-	Size = *(const uint16*)(Read(2));
+
+	Input.Read(PreambleCursor, MetadataSize, this, OpMetadataRead);
+
+	return true;
+}
+
+////////////////////////////////////////////////////////////////////////////////
+bool FAsioRecorderRelay::ReadMetadata(int32 Size)
+{
+	// At this point Buffer          [magic][md_size][metadata][t_ver][p_ver]
+	// looks like this;              Buffer--------->PreambleCursor--------->
+	//                                               |---------Size---------|
+
+	// We want to consume [metadata] so some adjustment is required.
+	int32 ReadSize = Size - sizeof(VersionType);
+	const uint8* Cursor = PreambleCursor;
+	auto Read = [&] (int32 SizeToRead)
+	{
+		const uint8* Ptr = Cursor;
+		Cursor += SizeToRead;
+		ReadSize -= SizeToRead;
+		return Ptr;
+	};
 
 	// MetadataFields
-	while (Size >= 2)
+	while (ReadSize > 0)
 	{
 		struct {
 			uint8	Size;
@@ -131,9 +223,9 @@ bool FAsioRecorderRelay::ReadMetadata(int32 Size)
 		} MetadataField;
 		MetadataField = *(const decltype(MetadataField)*)(Read(sizeof(MetadataField)));
 
-		if (Size < MetadataField.Size)
+		if (ReadSize < MetadataField.Size)
 		{
-			break;
+			return false;
 		}
 
 		if (MetadataField.Id == 0) /* ControlPortFieldId */
@@ -141,9 +233,25 @@ bool FAsioRecorderRelay::ReadMetadata(int32 Size)
 			ControlPort = *(const uint16*)Cursor;
 		}
 
-		Size -= MetadataField.Size;
+		ReadSize -= MetadataField.Size;
 	}
 
+	// There should be no data left to consume if the metadata was well-formed
+	if (ReadSize != 0)
+	{
+		return false;
+	}
+
+	// Now we've a full preamble we are ready to write the trace.
+	if (!CreateTrace())
+	{
+		return false;
+	}
+
+	// Analysis needs the preamble too.
+	uint32 PreambleSize = uint32(ptrdiff_t(PreambleCursor - Buffer)) + Size;
+	Output->Write(Buffer, PreambleSize, this, OpFileWrite);
+
 	return true;
 }
 
@@ -158,22 +266,26 @@ void FAsioRecorderRelay::OnIoComplete(uint32 Id, int32 Size)
 
 	switch (Id)
 	{
-	case OpSocketReadMetadata:
-		ActiveReadOp = OpSocketRead;
+	case OpMagicRead:
+		if (!ReadMagic())
+		{
+			Close();
+		}
+		break;
+
+	case OpMetadataRead:
 		if (!ReadMetadata(Size))
 		{
 			Close();
-			return;
 		}
-		/* fallthrough */
+		break;
 
 	case OpSocketRead:
 		Output->Write(Buffer, Size, this, OpFileWrite);
 		break;
 
-	case OpStart:
 	case OpFileWrite:
-		Input.ReadSome(Buffer, BufferSize, this, ActiveReadOp);
+		Input.ReadSome(Buffer, BufferSize, this, OpSocketRead);
 		break;
 	}
 }
@@ -189,7 +301,7 @@ uint32 FAsioRecorder::FSession::GetId() const
 ////////////////////////////////////////////////////////////////////////////////
 uint32 FAsioRecorder::FSession::GetTraceId() const
 {
-	return TraceId;
+	return Relay->GetTraceId();
 }
 
 ////////////////////////////////////////////////////////////////////////////////
@@ -259,38 +371,7 @@ const FAsioRecorder::FSession* FAsioRecorder::GetSessionInfo(uint32 Index) const
 ////////////////////////////////////////////////////////////////////////////////
 bool FAsioRecorder::OnAccept(asio::ip::tcp::socket& Socket)
 {
-	FAsioStore::FNewTrace Trace = Store.CreateTrace();
-	if (Trace.Writeable == nullptr)
-	{
-		return true;
-	}
-
-#if PLATFORM_WINDOWS
-	// Trace data is a stream and communication is one way. It is implemented
-	// this way to share code between sending trace data over the wire and writing
-	// it to a file. Because there's no ping/pong we can end up with a half-open
-	// TCP connection if the other end doesn't close its socket. So we'll enable
-	// keep-alive on the socket and set a short timeout (default is 2hrs).
-	tcp_keepalive KeepAlive =
-	{
-		1,		// on
-		15000,	// timeout_ms
-		2000,	// interval_ms
-	};
-
-	DWORD BytesReturned;
-	WSAIoctl(
-		Socket.native_handle(),
-		SIO_KEEPALIVE_VALS,
-		&KeepAlive, sizeof(KeepAlive),
-		nullptr, 0,
-		&BytesReturned,
-		nullptr,
-		nullptr
-	);
-#endif
-
-	auto* Relay = new FAsioRecorderRelay(Socket, Trace.Writeable);
+	auto* Relay = new FAsioRecorderRelay(Socket, Store);
 
 	uint32 IdPieces[] = {
 		Relay->GetIpAddress(),
@@ -302,7 +383,6 @@ bool FAsioRecorder::OnAccept(asio::ip::tcp::socket& Socket)
 	FSession Session;
 	Session.Relay = Relay;
 	Session.Id = QuickStoreHash(IdPieces);
-	Session.TraceId = Trace.Id;
 	Sessions.Add(Session);
 
 	return true;
diff --git a/Engine/Source/Developer/TraceAnalysis/Private/Store/AsioRecorder.h b/Engine/Source/Developer/TraceAnalysis/Private/Store/AsioRecorder.h
--- a/Engine/Source/Developer/TraceAnalysis/Private/Store/AsioRecorder.h
+++ b/Engine/Source/Developer/TraceAnalysis/Private/Store/AsioRecorder.h
@@ -31,7 +31,6 @@ public:
 		friend					FAsioRecorder;
 		FAsioRecorderRelay*		Relay;
 		uint32					Id;
-		uint32					TraceId;
 	};
 
 								FAsioRecorder(asio::io_context& IoContext, FAsioStore& InStore);
diff --git a/Engine/Source/Developer/TraceAnalysis/Private/Store/AsioStoreCborServer.cpp b/Engine/Source/Developer/TraceAnalysis/Private/Store/AsioStoreCborServer.cpp
--- a/Engine/Source/Developer/TraceAnalysis/Private/Store/AsioStoreCborServer.cpp
+++ b/Engine/Source/Developer/TraceAnalysis/Private/Store/AsioStoreCborServer.cpp
@@ -289,6 +289,11 @@ void FAsioStoreCborPeer::OnIoComplete(uint32 Id, int32 Size)
 		break;
 
 	case OpReadPayloadSize:
+		if (PayloadSize == 0 || PayloadSize > 1024)
+		{
+			SendError(EStatusCode::BadRequest);
+			break;
+		}
 		Socket.Read(Response.Reserve(PayloadSize), PayloadSize, this, OpReadPayload);
 		break;
 
