[Insights]
  - Added workaround for large utrace streams with corrupted serial sync events (detects when too much data is buffered on one thread and stops analysis automatically).
  - Changed the read buffers to increase their size in power of two steps.

diff --git a/Engine/Source/Developer/TraceAnalysis/Private/Analysis/Engine.cpp b/Engine/Source/Developer/TraceAnalysis/Private/Analysis/Engine.cpp
--- a/Engine/Source/Developer/TraceAnalysis/Private/Analysis/Engine.cpp
+++ b/Engine/Source/Developer/TraceAnalysis/Private/Analysis/Engine.cpp
@@ -3119,6 +3119,14 @@ FProtocol5Stage::EStatus FProtocol5Stage::OnDataNormal(const FMachineContext& Co
 		
 		// Extract all the events in the stream for this thread
 		FStreamReader* ThreadReader = Transport.GetThreadStream(i);
+
+		// Stop analysis if this thread has accumulated too much data.
+		// This can happen on corrupted traces (ex. with serial sync events missing or out of order).
+		if (ThreadReader->GetRemaining() > 512 * 1024 * 1024)
+		{
+			return EStatus::Error;
+		}
+
 		if (ParseEvents(*ThreadReader, EventDescs) < 0)
 		{
 			return EStatus::Error;
diff --git a/Engine/Source/Developer/TraceAnalysis/Private/Analysis/StreamReader.cpp b/Engine/Source/Developer/TraceAnalysis/Private/Analysis/StreamReader.cpp
--- a/Engine/Source/Developer/TraceAnalysis/Private/Analysis/StreamReader.cpp
+++ b/Engine/Source/Developer/TraceAnalysis/Private/Analysis/StreamReader.cpp
@@ -101,13 +101,13 @@ uint8* FStreamBuffer::Append(uint32 Size)
 ////////////////////////////////////////////////////////////////////////////////
 void FStreamBuffer::Consolidate()
 {
-	int32 Remaining = End - Cursor;
+	int32 Remaining = GetRemaining();
 	DemandHint += Remaining;
 
 	if (DemandHint >= BufferSize)
 	{
-		const uint32 GrowthSizeMask = (64 << 10) - 1;
-		BufferSize = (DemandHint + GrowthSizeMask + 1) & ~GrowthSizeMask;
+		check(DemandHint < (1u << 31));
+		BufferSize = FMath::Max(64u << 10, FMath::RoundUpToPowerOfTwo(DemandHint + 1));
 		Buffer = (uint8*)FMemory::Realloc(Buffer, BufferSize);
 	}
 
