[Insights] Surface error messages from analysis engine to the user

We have previously relied on UE_LOG messages to indicate problems in analysis. This change converts those log messages to use MessageLog functionality which is routed to a visual log representation inside the application (but also the log file).

Additionally adds functionality to show the Insights log windows when ever (new) errors are encountered.

diff --git a/Engine/Source/Developer/TraceAnalysis/Private/Analysis/Engine.cpp b/Engine/Source/Developer/TraceAnalysis/Private/Analysis/Engine.cpp
--- a/Engine/Source/Developer/TraceAnalysis/Private/Analysis/Engine.cpp
+++ b/Engine/Source/Developer/TraceAnalysis/Private/Analysis/Engine.cpp
@@ -1,26 +1,29 @@
 // Copyright Epic Games, Inc. All Rights Reserved.
 
 #include "Engine.h"
+#include "CoreGlobals.h"
+#include "StreamReader.h"
 #include "Algo/BinarySearch.h"
 #include "Algo/Sort.h"
 #include "Algo/StableSort.h"
 #include "Containers/ArrayView.h"
-#include "CoreGlobals.h"
 #include "HAL/UnrealMemory.h"
+#include "Internationalization/Internationalization.h"
 #include "Logging/LogMacros.h"
-#include "StreamReader.h"
+#include "Logging/MessageLog.h"
 #include "Templates/UnrealTemplate.h"
-#include "Trace/Analysis.h"
 #include "Trace/Analyzer.h"
 #include "Trace/Detail/Protocol.h"
 #include "Trace/Detail/Transport.h"
 #include "Transport/PacketTransport.h"
-#include "Transport/Transport.h"
 #include "Transport/TidPacketTransport.h"
+#include "Transport/Transport.h"
 
 namespace UE {
 namespace Trace {
 
+#define LOCTEXT_NAMESPACE "TraceAnalysis"
+
 // {{{1 misc -------------------------------------------------------------------
 
 ////////////////////////////////////////////////////////////////////////////////
@@ -1976,6 +1979,7 @@ public:
 	{
 		FAnalysisMachine&	Machine;
 		FAnalysisBridge&	Bridge;
+		FMessageLog*		Log;
 	};
 
 	class FStage
@@ -1990,7 +1994,7 @@ public:
 		virtual void		ExitStage(const FMachineContext& Context) {};
 	};
 
-							FAnalysisMachine(FAnalysisBridge& InBridge);
+							FAnalysisMachine(FAnalysisBridge& InBridge, FMessageLog* Log);
 							~FAnalysisMachine();
 	EStatus					OnData(FStreamReader& Reader);
 	void					Transition();
@@ -2003,11 +2007,13 @@ private:
 	FStage*					ActiveStage = nullptr;
 	TArray<FStage*>			StageQueue;
 	TArray<FStage*>			DeadStages;
+	FMessageLog*			Log;
 };
 
 ////////////////////////////////////////////////////////////////////////////////
-FAnalysisMachine::FAnalysisMachine(FAnalysisBridge& InBridge)
+FAnalysisMachine::FAnalysisMachine(FAnalysisBridge& InBridge, FMessageLog* InLog)
 : Bridge(InBridge)
+, Log(InLog)
 {
 }
 
@@ -2041,7 +2047,7 @@ void FAnalysisMachine::Transition()
 {
 	if (ActiveStage != nullptr)
 	{
-		FMachineContext Context = { *this, Bridge };
+		const FMachineContext Context = { *this, Bridge, Log };
 		ActiveStage->ExitStage(Context);
 
 		DeadStages.Add(ActiveStage);
@@ -2051,7 +2057,7 @@ void FAnalysisMachine::Transition()
 
 	if (ActiveStage != nullptr)
 	{
-		FMachineContext Context = { *this, Bridge };
+		const FMachineContext Context = { *this, Bridge, Log };
 		ActiveStage->EnterStage(Context);
 	}
 }
@@ -2059,7 +2065,7 @@ void FAnalysisMachine::Transition()
 ////////////////////////////////////////////////////////////////////////////////
 FAnalysisMachine::EStatus FAnalysisMachine::OnData(FStreamReader& Reader)
 {
-	FMachineContext Context = { *this, Bridge };
+	const FMachineContext Context = { *this, Bridge, Log };
 	EStatus Ret;
 	do
 	{
@@ -2874,12 +2880,12 @@ protected:
 	EStatus					OnDataImportant(const FMachineContext& Context);
 	EStatus					OnDataNonCachedImportant(const FMachineContext& Context);
 	EStatus					OnDataNormal(const FMachineContext& Context);
-	int32					ParseImportantEvents(FStreamReader& Reader, EventDescArray& OutEventDescs);
-	int32					ParseEvents(FStreamReader& Reader, EventDescArray& OutEventDescs);
-	int32					ParseEventsWithAux(FStreamReader& Reader, EventDescArray& OutEventDescs);
-	int32					ParseEvent(FStreamReader& Reader, FEventDesc& OutEventDesc);
-	int32					DispatchEvents(FAnalysisBridge& Bridge, const FEventDesc* EventDesc, uint32 Count);
-	int32					DispatchEvents(FAnalysisBridge& Bridge, TArray<FEventDescStream>& EventDescHeap);
+	int32					ParseImportantEvents(FStreamReader& Reader, EventDescArray& OutEventDescs, const FMachineContext& Context);
+	int32					ParseEvents(FStreamReader& Reader, EventDescArray& OutEventDescs, const FMachineContext& Context);
+	int32					ParseEventsWithAux(FStreamReader& Reader, EventDescArray& OutEventDescs, const FMachineContext& Context);
+	int32					ParseEvent(FStreamReader& Reader, FEventDesc& OutEventDesc, const FMachineContext& Context);
+	int32					DispatchEvents(const FMachineContext& Context, const FEventDesc* EventDesc, uint32 Count);
+	int32					DispatchEvents(const FMachineContext& Context, TArray<FEventDescStream>& EventDescHeap);
 	void					DetectSerialGaps(TArray<FEventDescStream>& EventDescHeap);
 	template <typename Callback>
 	void					ForEachSerialGap(const TArray<FEventDescStream>& EventDescHeap, Callback&& InCallback);
@@ -2964,7 +2970,7 @@ FProtocol5Stage::EStatus FProtocol5Stage::OnDataNewEvents(const FMachineContext&
 		return EStatus::EndOfStream;
 	}
 
-	if (ParseImportantEvents(*ThreadReader, EventDescs) < 0)
+	if (ParseImportantEvents(*ThreadReader, EventDescs, Context) < 0)
 	{
 		return EStatus::Error;
 	}
@@ -2991,7 +2997,7 @@ FProtocol5Stage::EStatus FProtocol5Stage::OnDataImportant(const FMachineContext&
 		return EStatus::EndOfStream;
 	}
 
-	if (ParseImportantEvents(*ThreadReader, EventDescs) < 0)
+	if (ParseImportantEvents(*ThreadReader, EventDescs, Context) < 0)
 	{
 		return EStatus::Error;
 	}
@@ -3009,7 +3015,7 @@ FProtocol5Stage::EStatus FProtocol5Stage::OnDataImportant(const FMachineContext&
 	EventDesc.Serial = ESerial::Terminal;
 
 	Context.Bridge.SetActiveThread(ETransportTid::Importants);
-	if (DispatchEvents(Context.Bridge, EventDescs.GetData(), EventDescs.Num() - 1) < 0)
+	if (DispatchEvents(Context, EventDescs.GetData(), EventDescs.Num() - 1) < 0)
 	{
 		return EStatus::Error;
 	}
@@ -3018,7 +3024,7 @@ FProtocol5Stage::EStatus FProtocol5Stage::OnDataImportant(const FMachineContext&
 }
 
 ////////////////////////////////////////////////////////////////////////////////
-int32 FProtocol5Stage::ParseImportantEvents(FStreamReader& Reader, EventDescArray& OutEventDescs)
+int32 FProtocol5Stage::ParseImportantEvents(FStreamReader& Reader, EventDescArray& OutEventDescs, const FMachineContext& Context)
 {
 	using namespace Protocol5;
 
@@ -3055,7 +3061,7 @@ int32 FProtocol5Stage::ParseImportantEvents(FStreamReader& Reader, EventDescArra
 		const FTypeRegistry::FTypeInfo* TypeInfo = TypeRegistry.Get(Uid);
 		if (TypeInfo == nullptr)
 		{
-			UE_LOG(LogTraceAnalysis, Log, TEXT("Unknown event UID: %08x"), Uid);
+			Context.Log->Warning(LOCTEXT("UnknownUID","An unknown event UID was encountered."));
 			return 1;
 		}
 
@@ -3089,7 +3095,7 @@ int32 FProtocol5Stage::ParseImportantEvents(FStreamReader& Reader, EventDescArra
 
 			if (Cursor[0] != uint8(EKnownUids::AuxDataTerminal))
 			{
-				UE_LOG(LogTraceAnalysis, Warning, TEXT("Expecting AuxDataTerminal event"));
+				Context.Log->Warning(LOCTEXT("AuxDataTerminalExpected","Expected an aux data terminal in the stream."));
 				return -1;
 			}
 		}
@@ -3124,10 +3130,11 @@ FProtocol5Stage::EStatus FProtocol5Stage::OnDataNormal(const FMachineContext& Co
 		// This can happen on corrupted traces (ex. with serial sync events missing or out of order).
 		if (ThreadReader->GetRemaining() > 512 * 1024 * 1024)
 		{
+			Context.Log->Error(LOCTEXT("TooMuchData", "Analysis accumulated too much data before being able to continue analysing the stream."));
 			return EStatus::Error;
 		}
 
-		if (ParseEvents(*ThreadReader, EventDescs) < 0)
+		if (ParseEvents(*ThreadReader, EventDescs, Context) < 0)
 		{
 			return EStatus::Error;
 		}
@@ -3171,7 +3178,7 @@ FProtocol5Stage::EStatus FProtocol5Stage::OnDataNormal(const FMachineContext& Co
 		{
 			Context.Bridge.SetActiveThread(Stream.ThreadId);
 
-			if (DispatchEvents(Context.Bridge, StartDesc, DescNum) < 0)
+			if (DispatchEvents(Context, StartDesc, DescNum) < 0)
 			{
 				return EStatus::Error;
 			}
@@ -3209,7 +3216,7 @@ FProtocol5Stage::EStatus FProtocol5Stage::OnDataNormal(const FMachineContext& Co
 	const bool bSync = (SyncCount != Transport.GetSyncCount());
 	DetectSerialGaps(EventDescHeap);
 
-	if (DispatchEvents(Context.Bridge, EventDescHeap) < 0)
+	if (DispatchEvents(Context, EventDescHeap) < 0)
 	{
 		return EStatus::Error;
 	}
@@ -3224,7 +3231,7 @@ FProtocol5Stage::EStatus FProtocol5Stage::OnDataNormal(const FMachineContext& Co
 
 ////////////////////////////////////////////////////////////////////////////////
 int32 FProtocol5Stage::DispatchEvents(
-	FAnalysisBridge& Bridge,
+	const FMachineContext& Context,
 	TArray<FEventDescStream>& EventDescHeap)
 {
 	auto UpdateHeap = [&] (const FEventDescStream& Stream, const FEventDesc* EventDesc)
@@ -3279,10 +3286,10 @@ int32 FProtocol5Stage::DispatchEvents(
 		}
 
 		// Dispatch.
-		Bridge.SetActiveThread(Stream.ThreadId);
+		Context.Bridge.SetActiveThread(Stream.ThreadId);
 		int32 DescNum = int32(UPTRINT(EndDesc - StartDesc));
 		check(DescNum > 0);
-		if (DispatchEvents(Bridge, StartDesc, DescNum) < 0)
+		if (DispatchEvents(Context, StartDesc, DescNum) < 0)
 		{
 			return -1;
 		}
@@ -3307,12 +3314,12 @@ int32 FProtocol5Stage::DispatchEvents(
 }
 
 ////////////////////////////////////////////////////////////////////////////////
-int32 FProtocol5Stage::ParseEvents(FStreamReader& Reader, EventDescArray& OutEventDescs)
+int32 FProtocol5Stage::ParseEvents(FStreamReader& Reader, EventDescArray& OutEventDescs, const FMachineContext& Context)
 {
 	while (!Reader.IsEmpty())
 	{
 		FEventDesc EventDesc;
-		int32 Size = ParseEvent(Reader, EventDesc);
+		int32 Size = ParseEvent(Reader, EventDesc, Context);
 		if (Size <= 0)
 		{
 			return Size;
@@ -3328,7 +3335,7 @@ int32 FProtocol5Stage::ParseEvents(FStreamReader& Reader, EventDescArray& OutEve
 
 			Reader.Advance(Size);
 
-			int Ok = ParseEventsWithAux(Reader, OutEventDescs);
+			int Ok = ParseEventsWithAux(Reader, OutEventDescs, Context);
 			if (Ok < 0)
 			{
 				return Ok;
@@ -3351,7 +3358,7 @@ int32 FProtocol5Stage::ParseEvents(FStreamReader& Reader, EventDescArray& OutEve
 }
 
 ////////////////////////////////////////////////////////////////////////////////
-int32 FProtocol5Stage::ParseEventsWithAux(FStreamReader& Reader, EventDescArray& OutEventDescs)
+int32 FProtocol5Stage::ParseEventsWithAux(FStreamReader& Reader, EventDescArray& OutEventDescs, const FMachineContext& Context)
 {
 	// We are now "in" the scope of an event with zero or more aux-data blocks.
 	// We will consume events until we leave this scope (a aux-data-terminal).
@@ -3371,7 +3378,7 @@ int32 FProtocol5Stage::ParseEventsWithAux(FStreamReader& Reader, EventDescArray&
 		FEventDesc EventDesc;
 		EventDesc.Serial = ESerial::Ignored;
 
-		int32 Size = ParseEvent(Reader, EventDesc);
+		int32 Size = ParseEvent(Reader, EventDesc, Context);
 		if (Size <= 0)
 		{
 			return Size;
@@ -3440,7 +3447,7 @@ int32 FProtocol5Stage::ParseEventsWithAux(FStreamReader& Reader, EventDescArray&
 }
 
 ////////////////////////////////////////////////////////////////////////////////
-int32 FProtocol5Stage::ParseEvent(FStreamReader& Reader, FEventDesc& EventDesc)
+int32 FProtocol5Stage::ParseEvent(FStreamReader& Reader, FEventDesc& EventDesc, const FMachineContext& Context)
 {
 	using namespace Protocol5;
 
@@ -3500,7 +3507,7 @@ int32 FProtocol5Stage::ParseEvent(FStreamReader& Reader, FEventDesc& EventDesc)
 		const FTypeRegistry::FTypeInfo* TypeInfo = TypeRegistry.Get(Uid);
 		if (TypeInfo == nullptr)
 		{
-			UE_LOG(LogTraceAnalysis, Log, TEXT("Unknown event UID: %08x"), Uid);
+			Context.Log->Warning(LOCTEXT("UnknownUID","An unknown event UID was encountered."));	
 			return 0;
 		}
 
@@ -3680,7 +3687,7 @@ void FProtocol5Stage::DetectSerialGaps(TArray<FEventDescStream>& EventDescHeap)
 
 ////////////////////////////////////////////////////////////////////////////////
 int32 FProtocol5Stage::DispatchEvents(
-	FAnalysisBridge& Bridge,
+	const FMachineContext& Context,
 	const FEventDesc* EventDesc, uint32 Count)
 {
 	using namespace Protocol5;
@@ -3695,19 +3702,19 @@ int32 FProtocol5Stage::DispatchEvents(
 		case EKnownUids::EnterScope_T:
 		{
 			uint64 Timestamp = *(uint64*)(Cursor->Data - 1) >> 8;
-			Bridge.EnterScope(Timestamp);
+			Context.Bridge.EnterScope(Timestamp);
 			continue;
 		}
 
 		case EKnownUids::LeaveScope_T:
 		{
 			uint64 Timestamp = *(uint64*)(Cursor->Data - 1) >> 8;
-			Bridge.LeaveScope(Timestamp);
+			Context.Bridge.LeaveScope(Timestamp);
 			continue;
 		}
 
-		case EKnownUids::EnterScope: Bridge.EnterScope(); continue;
-		case EKnownUids::LeaveScope: Bridge.LeaveScope(); continue;
+		case EKnownUids::EnterScope: Context.Bridge.EnterScope(); continue;
+		case EKnownUids::LeaveScope: Context.Bridge.LeaveScope(); continue;
 
 		case EKnownUids::AuxData:
 		case EKnownUids::AuxDataTerminal:
@@ -3716,7 +3723,7 @@ int32 FProtocol5Stage::DispatchEvents(
 		default:
 			if (!TypeRegistry.IsUidValid(Cursor->Uid))
 			{
-				UE_LOG(LogTraceAnalysis, Warning, TEXT("Unexpected event UID: %08x"), Cursor->Uid);
+				Context.Log->Warning(LOCTEXT("UnknownUID","An unknown event UID was encountered."));	
 				return -1;
 			}
 			break;
@@ -3755,7 +3762,7 @@ int32 FProtocol5Stage::DispatchEvents(
 			}
 		}
 
-		Bridge.OnEvent(EventDataInfo);
+		Context.Bridge.OnEvent(EventDataInfo);
 
 		AuxCollector.Reset();
 	}
@@ -3815,7 +3822,13 @@ FEstablishTransportStage::EStatus FEstablishTransportStage::OnData(
 	case ETransport::Packet:		Transport = new FPacketTransport(); break;
 	case ETransport::TidPacket:		Transport = new FTidPacketTransport(); break;
 	case ETransport::TidPacketSync:	Transport = new FTidPacketTransportSync(); break;
-	default:						return EStatus::Error;
+	default:
+		{
+			Context.Log->Error(FText::Format(
+				LOCTEXT("InvalidTransportVersion", "Unknown transport version: {0}. You may need to recompile this application"),
+				FText::AsNumber(Header->TransportVersion)));
+			return EStatus::Error;
+		}
 	}
 
 	uint32 ProtocolVersion = Header->ProtocolVersion;
@@ -3849,7 +3862,12 @@ FEstablishTransportStage::EStatus FEstablishTransportStage::OnData(
 		break;
 
 	default:
-		return EStatus::Error;
+		{
+			Context.Log->Error(FText::Format(
+					LOCTEXT("UnknownProtocolVersion", "Unknown protocol version: {0}. You may need to recompile this application."),
+					FText::AsNumber(ProtocolVersion)));
+			return EStatus::Error;
+		}
 	}
 
 	Reader.Advance(sizeof(*Header));
@@ -3920,6 +3938,7 @@ FMagicStage::EStatus FMagicStage::OnData(
 	if (Magic == 'ECRT' || Magic == '2CRT')
 	{
 		// Source is big-endian which we don't currently support
+		Context.Log->Error(LOCTEXT("BigEndian", "Big endian traces are currently not supported."));
 		return EStatus::Error;
 	}
 
@@ -3948,6 +3967,7 @@ FMagicStage::EStatus FMagicStage::OnData(
 		return EStatus::Continue;
 	}
 
+	Context.Log->Error(LOCTEXT("IncorrectMagic", "The file or stream was not recognized as trace stream."));
 	return EStatus::Error;
 }
 
@@ -3959,17 +3979,18 @@ FMagicStage::EStatus FMagicStage::OnData(
 class FAnalysisEngine::FImpl
 {
 public:
-						FImpl(TArray<IAnalyzer*>&& Analyzers);
+					FImpl(TArray<IAnalyzer*>&& Analyzers, FMessageLog* InLog);
 	void				Begin();
 	void				End();
 	bool				OnData(FStreamReader& Reader);
 	FAnalysisBridge		Bridge;
-	FAnalysisMachine	Machine = Bridge;
+	FAnalysisMachine	Machine;
 };
 
 ////////////////////////////////////////////////////////////////////////////////
-FAnalysisEngine::FImpl::FImpl(TArray<IAnalyzer*>&& Analyzers)
+FAnalysisEngine::FImpl::FImpl(TArray<IAnalyzer*>&& Analyzers, FMessageLog* InLog)
 : Bridge(Forward<TArray<IAnalyzer*>>(Analyzers))
+, Machine(Bridge, InLog)
 {
 }
 
@@ -3998,8 +4019,8 @@ bool FAnalysisEngine::FImpl::OnData(FStreamReader& Reader)
 
 
 ////////////////////////////////////////////////////////////////////////////////
-FAnalysisEngine::FAnalysisEngine(TArray<IAnalyzer*>&& Analyzers)
-: Impl(new FImpl(Forward<TArray<IAnalyzer*>>(Analyzers)))
+FAnalysisEngine::FAnalysisEngine(TArray<IAnalyzer*>&& Analyzers, FMessageLog* InLog)
+: Impl(new FImpl(Forward<TArray<IAnalyzer*>>(Analyzers), InLog))
 {
 }
 
@@ -4031,4 +4052,6 @@ bool FAnalysisEngine::OnData(FStreamReader& Reader)
 } // namespace Trace
 } // namespace UE
 
+#undef LOCTEXT_NAMESPACE
+
 /* vim: set foldlevel=1 : */
diff --git a/Engine/Source/Developer/TraceAnalysis/Private/Analysis/Engine.h b/Engine/Source/Developer/TraceAnalysis/Private/Analysis/Engine.h
--- a/Engine/Source/Developer/TraceAnalysis/Private/Analysis/Engine.h
+++ b/Engine/Source/Developer/TraceAnalysis/Private/Analysis/Engine.h
@@ -4,6 +4,8 @@
 
 #include "Containers/ContainersFwd.h"
 
+class FMessageLog;
+
 namespace UE {
 namespace Trace {
 
@@ -14,7 +16,7 @@ class FStreamReader;
 class FAnalysisEngine
 {
 public:
-						FAnalysisEngine(TArray<IAnalyzer*>&& Analyzers);
+						FAnalysisEngine(TArray<IAnalyzer*>&& Analyzers, FMessageLog* InLog);
 						~FAnalysisEngine();
 	void				Begin();
 	void				End();
diff --git a/Engine/Source/Developer/TraceAnalysis/Private/Analysis/Processor.cpp b/Engine/Source/Developer/TraceAnalysis/Private/Analysis/Processor.cpp
--- a/Engine/Source/Developer/TraceAnalysis/Private/Analysis/Processor.cpp
+++ b/Engine/Source/Developer/TraceAnalysis/Private/Analysis/Processor.cpp
@@ -6,6 +6,8 @@
 #include "HAL/PlatformProcess.h"
 #include "HAL/RunnableThread.h"
 #include "StreamReader.h"
+#include "TraceAnalysisModule.h"
+#include "Logging/MessageLog.h"
 #include "Templates/UnrealTemplate.h"
 #include "Trace/DataStream.h"
 
@@ -14,7 +16,8 @@ namespace Trace {
 
 ////////////////////////////////////////////////////////////////////////////////
 FAnalysisProcessor::FImpl::FImpl(IInDataStream& InDataStream, TArray<IAnalyzer*>&& InAnalyzers)
-: AnalysisEngine(Forward<TArray<IAnalyzer*>>(InAnalyzers))
+: Log(FTraceAnalysisModule::GetMessageLogName())
+, AnalysisEngine(Forward<TArray<IAnalyzer*>>(InAnalyzers), &Log)
 , DataStream(InDataStream)
 , StopEvent(FPlatformProcess::GetSynchEventFromPool(true))
 , UnpausedEvent(FPlatformProcess::GetSynchEventFromPool(true))
@@ -100,6 +103,13 @@ void FAnalysisProcessor::FImpl::PauseAnalysis(bool bState)
 	}
 }
 
+////////////////////////////////////////////////////////////////////////////////
+FMessageLog* FAnalysisProcessor::FImpl::GetLog() 
+{
+	// MessageLog is not thread safe, so it's not safe to return it while
+	// messages could be pushed to the log.
+	return bComplete ? &Log : nullptr;
+}
 
 
 ////////////////////////////////////////////////////////////////////////////////
@@ -107,6 +117,7 @@ bool FAnalysisProcessor::IsActive() const	{ return (Impl != nullptr) ? Impl->IsA
 void FAnalysisProcessor::Stop()				{ if (Impl != nullptr) { Impl->StopAnalysis(); } }
 void FAnalysisProcessor::Wait()				{ if (Impl != nullptr) { Impl->WaitOnAnalysis(); } }
 void FAnalysisProcessor::Pause(bool bState) { if (Impl != nullptr) { Impl->PauseAnalysis(bState); } }
+FMessageLog* FAnalysisProcessor::GetLog() { return Impl->GetLog(); }	
 
 ////////////////////////////////////////////////////////////////////////////////
 FAnalysisProcessor::FAnalysisProcessor(FAnalysisProcessor&& Rhs)
diff --git a/Engine/Source/Developer/TraceAnalysis/Private/Analysis/Processor.h b/Engine/Source/Developer/TraceAnalysis/Private/Analysis/Processor.h
--- a/Engine/Source/Developer/TraceAnalysis/Private/Analysis/Processor.h
+++ b/Engine/Source/Developer/TraceAnalysis/Private/Analysis/Processor.h
@@ -4,6 +4,7 @@
 
 #include "Engine.h"
 #include "HAL/Runnable.h"
+#include "Logging/MessageLog.h"
 #include "Trace/Analysis.h"
 
 class FEvent;
@@ -24,8 +25,10 @@ public:
 	void				StopAnalysis();
 	void				WaitOnAnalysis();
 	void				PauseAnalysis(bool bState);
+	FMessageLog*		GetLog();
 
 private:
+	FMessageLog			Log;
 	FAnalysisEngine		AnalysisEngine;
 	IInDataStream&		DataStream;
 	FEvent*				StopEvent;
diff --git a/Engine/Source/Developer/TraceAnalysis/Private/TraceAnalysisModule.cpp b/Engine/Source/Developer/TraceAnalysis/Private/TraceAnalysisModule.cpp
--- a/Engine/Source/Developer/TraceAnalysis/Private/TraceAnalysisModule.cpp
+++ b/Engine/Source/Developer/TraceAnalysis/Private/TraceAnalysisModule.cpp
@@ -1,5 +1,24 @@
 // Copyright Epic Games, Inc. All Rights Reserved.
 
+#include "TraceAnalysisModule.h"
+
+#include "MessageLog/Public/MessageLogModule.h"
 #include "Modules/ModuleManager.h"
 
-IMPLEMENT_MODULE(FDefaultModuleImpl, TraceAnalysis);
+////////////////////////////////////////////////////////////////////////////////
+void FTraceAnalysisModule::StartupModule()
+{
+	FMessageLogModule& MessageLogModule = FModuleManager::LoadModuleChecked<FMessageLogModule>("MessageLog");
+	MessageLogModule.RegisterLogListing(GetMessageLogName(), NSLOCTEXT("Trace", "TraceAnalysis", "Trace Analysis"));
+	MessageLogModule.EnableMessageLogDisplay(true);
+}
+
+////////////////////////////////////////////////////////////////////////////////
+FName FTraceAnalysisModule::GetMessageLogName()
+{
+	static FName Name("TraceAnalysis");
+	return Name;
+}
+
+////////////////////////////////////////////////////////////////////////////////
+IMPLEMENT_MODULE(FTraceAnalysisModule, TraceAnalysis);
diff --git a/Engine/Source/Developer/TraceAnalysis/Private/TraceAnalysisModule.h b/Engine/Source/Developer/TraceAnalysis/Private/TraceAnalysisModule.h
new file mode 100644
--- /dev/null
+++ b/Engine/Source/Developer/TraceAnalysis/Private/TraceAnalysisModule.h
@@ -0,0 +1,12 @@
+﻿// Copyright Epic Games, Inc. All Rights Reserved.
+#pragma once
+#include "CoreMinimal.h"
+#include "Modules/ModuleInterface.h"
+
+class FTraceAnalysisModule : public IModuleInterface
+{
+public:
+	virtual void StartupModule() override;
+
+	static FName GetMessageLogName();
+};
\ No newline at end of file
diff --git a/Engine/Source/Developer/TraceAnalysis/Public/Trace/Analysis.h b/Engine/Source/Developer/TraceAnalysis/Public/Trace/Analysis.h
--- a/Engine/Source/Developer/TraceAnalysis/Public/Trace/Analysis.h
+++ b/Engine/Source/Developer/TraceAnalysis/Public/Trace/Analysis.h
@@ -32,6 +32,12 @@ public:
 	 * @param bState Pause if true, resume if false. */
 	void Pause(bool bState);
 
+
+	/**
+	 * Gets the message log. Must be called after analysis has finished.
+	 */
+	FMessageLog* GetLog();
+
 						~FAnalysisProcessor();
 						FAnalysisProcessor() = default;
 						FAnalysisProcessor(FAnalysisProcessor&& Rhs);
diff --git a/Engine/Source/Developer/TraceServices/Private/AnalysisService.cpp b/Engine/Source/Developer/TraceServices/Private/AnalysisService.cpp
--- a/Engine/Source/Developer/TraceServices/Private/AnalysisService.cpp
+++ b/Engine/Source/Developer/TraceServices/Private/AnalysisService.cpp
@@ -129,6 +129,11 @@ void FAnalysisSession::Wait() const
 	Processor.Wait();
 }
 
+FMessageLog* FAnalysisSession::GetLog() const
+{
+	return Processor.GetLog();
+}
+
 void FAnalysisSession::AddAnalyzer(UE::Trace::IAnalyzer* Analyzer)
 {
 	Analyzers.Add(Analyzer);
diff --git a/Engine/Source/Developer/TraceServices/Private/AnalysisServicePrivate.h b/Engine/Source/Developer/TraceServices/Private/AnalysisServicePrivate.h
--- a/Engine/Source/Developer/TraceServices/Private/AnalysisServicePrivate.h
+++ b/Engine/Source/Developer/TraceServices/Private/AnalysisServicePrivate.h
@@ -46,6 +46,7 @@ public:
 	virtual bool IsAnalysisComplete() const override { return !Processor.IsActive(); }
 	virtual double GetDurationSeconds() const override { Lock.ReadAccessCheck(); return DurationSeconds; }
 	virtual void UpdateDurationSeconds(double Duration) override { Lock.WriteAccessCheck(); DurationSeconds = FMath::Max(Duration, DurationSeconds); }
+	virtual FMessageLog* GetLog() const override;
 
 	virtual ILinearAllocator& GetLinearAllocator() override { return Allocator; }
 	virtual const TCHAR* StoreString(const TCHAR* String) override { return StringStore.Store(String); }
diff --git a/Engine/Source/Developer/TraceServices/Private/TraceServicesModule.cpp b/Engine/Source/Developer/TraceServices/Private/TraceServicesModule.cpp
--- a/Engine/Source/Developer/TraceServices/Private/TraceServicesModule.cpp
+++ b/Engine/Source/Developer/TraceServices/Private/TraceServicesModule.cpp
@@ -102,6 +102,9 @@ void FTraceServicesModule::StartupModule()
 {
 	LLM_SCOPE_BYTAG(Insights_TraceServices);
 
+	// Load the analysis module.
+	IModuleInterface& TraceAnalysisModule = FModuleManager::LoadModuleChecked<IModuleInterface>("TraceAnalysis");
+	
 	IModularFeatures::Get().RegisterModularFeature(TraceServices::ModuleFeatureName, &TimingProfilerModule);
 	IModularFeatures::Get().RegisterModularFeature(TraceServices::ModuleFeatureName, &CsvProfilerModule);
 	IModularFeatures::Get().RegisterModularFeature(TraceServices::ModuleFeatureName, &CountersModule);
diff --git a/Engine/Source/Developer/TraceServices/Public/TraceServices/Model/AnalysisSession.h b/Engine/Source/Developer/TraceServices/Public/TraceServices/Model/AnalysisSession.h
--- a/Engine/Source/Developer/TraceServices/Public/TraceServices/Model/AnalysisSession.h
+++ b/Engine/Source/Developer/TraceServices/Public/TraceServices/Model/AnalysisSession.h
@@ -8,6 +8,7 @@
 #include "Templates/SharedPointer.h"
 
 class FName;
+class FMessageLog;
 
 namespace UE {
 namespace Trace {
@@ -54,6 +55,7 @@ public:
 	virtual bool IsAnalysisComplete() const = 0;
 	virtual double GetDurationSeconds() const = 0;
 	virtual void UpdateDurationSeconds(double Duration) = 0;
+	virtual FMessageLog* GetLog() const = 0;
 
 	virtual void BeginRead() const = 0;
 	virtual void EndRead() const = 0;
