More robust analysis of trace streams where a large distance occurs between consecutive serialised events on different threads.
The main premise of this change is to order the heap of pending serialised events by distance from the next expected serial instead of in respect to each other. By using an unsigned comparison, wrapped serials that would otherwise appear on the negative side of the expected serial, in fact sort to the positive side. This allows the use of the full serial space as wraps no longer need to be calculated when sorting.
#rnx

diff --git a/Engine/Source/Developer/TraceAnalysis/Private/Analysis/Engine.cpp b/Engine/Source/Developer/TraceAnalysis/Private/Analysis/Engine.cpp
--- a/Engine/Source/Developer/TraceAnalysis/Private/Analysis/Engine.cpp
+++ b/Engine/Source/Developer/TraceAnalysis/Private/Analysis/Engine.cpp
@@ -2856,17 +2856,25 @@ protected:
 		};
 
 		enum { GapThreadId = ~0u };
-
-		bool operator < (const FEventDescStream& Rhs) const;
 	};
 	static_assert(sizeof(FEventDescStream) == 16, "");
 
+	struct FSerialDistancePredicate
+	{
+		bool operator () (const FEventDescStream& Lhs, const FEventDescStream& Rhs) const
+		{
+			uint32 Ld = Lhs.EventDescs->Serial - Origin;
+			uint32 Rd = Rhs.EventDescs->Serial - Origin;
+			return Ld < Rd;
+		};
+		uint32 Origin;
+	};
+
 	enum ESerial : int32 
 	{
 		Bits		= 24,
 		Mask		= (1 << Bits) - 1,
 		Range		= 1 << Bits,
-		HalfRange	= Range >> 1,
 		Ignored		= Range << 2, // far away so proper serials always compare less-than
 		Terminal,
 	};
@@ -2898,14 +2906,6 @@ protected:
 	uint32					EventVersion = 4; //Protocol version 5 uses the event version from protocol 4
 };
 
-////////////////////////////////////////////////////////////////////////////////
-bool FProtocol5Stage::FEventDescStream::operator < (const FEventDescStream& Rhs) const
-{
-	int32 Delta = Rhs.EventDescs->Serial - EventDescs->Serial;
-	int32 Wrapped = uint32(Delta + ESerial::HalfRange - 1) >= uint32(ESerial::Range - 2);
-	return (Wrapped ^ (Delta > 0)) != 0;
-}
-
 ////////////////////////////////////////////////////////////////////////////////
 FProtocol5Stage::FProtocol5Stage(FTransport* InTransport)
 : Transport(*(FTidPacketTransport*)InTransport)
@@ -3205,7 +3205,7 @@ FProtocol5Stage::EStatus FProtocol5Stage::OnDataNormal(const FMachineContext& Co
 	// than half the serial space, they have wrapped.
 
 	// A min-heap is used to peel off groups of events by lowest serial
-	EventDescHeap.Heapify();
+	EventDescHeap.Heapify(FSerialDistancePredicate{NextSerial});
 
 	// Events must be consumed contiguously.
 	if (NextSerial == ~0u && Transport.GetSyncCount())
@@ -3243,7 +3243,7 @@ int32 FProtocol5Stage::DispatchEvents(
 			EventDescHeap.Add(Next);
 		}
 
-		EventDescHeap.HeapPopDiscard();
+		EventDescHeap.HeapPopDiscard(FSerialDistancePredicate{NextSerial}, false);
 	};
 
 	do
@@ -3586,7 +3586,7 @@ void FProtocol5Stage::ForEachSerialGap(
 			auto& Out = HeapCopy.Add_GetRef({Stream.ThreadId, Stream.TransportIndex});
 			Out.EventDescs = EventDesc;
 		}
-		HeapCopy.HeapPopDiscard();
+		HeapCopy.HeapPopDiscard(FSerialDistancePredicate{NextSerial}, false);
 	}
 	while (!HeapCopy.IsEmpty());
 }
@@ -3681,7 +3681,7 @@ void FProtocol5Stage::DetectSerialGaps(TArray<FEventDescStream>& EventDescHeap)
 		FEventDescStream Out = EventDescHeap[0];
 		Out.ThreadId = FEventDescStream::GapThreadId;
 		Out.EventDescs = SerialGaps.GetData();
-		EventDescHeap.HeapPush(Out);
+		EventDescHeap.HeapPush(Out, FSerialDistancePredicate{NextSerial});
 	}
 }
 
