[Insights] Added analysis support for the new trace protocol 7 (scoped events with absolute timestamp values).

diff --git a/Engine/Source/Developer/TraceAnalysis/Private/Analysis/Engine.cpp b/Engine/Source/Developer/TraceAnalysis/Private/Analysis/Engine.cpp
--- a/Engine/Source/Developer/TraceAnalysis/Private/Analysis/Engine.cpp
+++ b/Engine/Source/Developer/TraceAnalysis/Private/Analysis/Engine.cpp
@@ -1,16 +1,17 @@
 // Copyright Epic Games, Inc. All Rights Reserved.
 
 #include "Engine.h"
-#include "CoreGlobals.h"
-#include "StreamReader.h"
+
 #include "Algo/BinarySearch.h"
 #include "Algo/Sort.h"
 #include "Algo/StableSort.h"
 #include "Containers/ArrayView.h"
+#include "CoreGlobals.h"
 #include "HAL/UnrealMemory.h"
 #include "Internationalization/Internationalization.h"
 #include "Logging/LogMacros.h"
 #include "Logging/MessageLog.h"
+#include "StreamReader.h"
 #include "Templates/UnrealTemplate.h"
 #include "Trace/Analyzer.h"
 #include "Trace/Detail/Protocol.h"
@@ -19,11 +20,20 @@
 #include "Transport/TidPacketTransport.h"
 #include "Transport/Transport.h"
 
-namespace UE {
-namespace Trace {
+DEFINE_LOG_CATEGORY_STATIC(LogTraceAnalysis, Log, All)
+
+#if 0
+#	define TRACE_ANALYSIS_DEBUG(Format, ...) \
+		do { UE_LOG(LogTraceAnalysis, Log, TEXT(Format), __VA_ARGS__) } while (0)
+#else
+#	define TRACE_ANALYSIS_DEBUG(...)
+#endif
 
 #define LOCTEXT_NAMESPACE "TraceAnalysis"
 
+namespace UE {
+namespace Trace {
+
 // {{{1 misc -------------------------------------------------------------------
 
 ////////////////////////////////////////////////////////////////////////////////
@@ -530,7 +540,7 @@ FDispatch::FField& FDispatchBuilder::AddField(const ANSICHAR* Name, int32 NameSi
 	FFnv1aHash Hash;
 	Hash.Add((const uint8*)Name, NameSize);
 	Field->Hash = Hash.Get();
-	
+
 	auto* Dispatch = (FDispatch*)(Buffer.GetData());
 	Dispatch->FieldCount++;
 	Dispatch->EventSize += Field->Size;
@@ -659,7 +669,7 @@ IAnalyzer::FEventFieldInfo::EType IAnalyzer::FEventFieldInfo::GetType() const
 			default: check(false); // Unsupported width
 		}
 	}
-	
+
 	if (Inner->Class == UE::Trace::Protocol0::Field_String)
 	{
 		return (Inner->SizeAndType == 1) ? EType::AnsiString : EType::WideString;
@@ -1041,6 +1051,7 @@ const FTypeRegistry::FTypeInfo* FTypeRegistry::Add(const void* TraceData, uint32
 	case 0: return AddVersion4(TraceData);
 	case 4: return AddVersion4(TraceData);
 	case 6: return AddVersion6(TraceData);
+	case 7: return AddVersion6(TraceData);
 	default:
 		check(false); // Unknown version
 	}
@@ -1172,7 +1183,7 @@ const FTypeRegistry::FTypeInfo* FTypeRegistry::AddVersion6(const void* TraceData
 		{
 			check((Field.DefinitionId.TypeInfo & Protocol0::Field_CategoryMask) == Protocol0::Field_Integer);
 			const int8 TypeSize = 1 << (Field.DefinitionId.TypeInfo & Protocol0::Field_Pow2SizeMask);
-			
+
 			auto DefinitionIdFieldName = ANSITEXTVIEW("DefinitionId");
 			auto& OutField = Builder.AddField(DefinitionIdFieldName.GetData(), DefinitionIdFieldName.Len(), TypeSize);
 			OutField.Offset = Field.DefinitionId.Offset;
@@ -1221,17 +1232,17 @@ void FTypeRegistry::Add(FTypeInfo* TypeInfo)
 	// but due to backwards compatibility we'll override existing types.
 	uint16 Uid = TypeInfo->Uid;
 	if (Uid < uint32(TypeInfos.Num()))
- 	{
+	{
 		if (TypeInfos[Uid] != nullptr)
- 		{
+		{
 			FMemory::Free(TypeInfos[Uid]);
 			TypeInfos[Uid] = nullptr;
- 		}
- 	}
+		}
+	}
 	else
- 	{
+	{
 		TypeInfos.SetNum(Uid + 1);
- 	}
+	}
 
 	TypeInfos[Uid] = TypeInfo;
 }
@@ -1382,9 +1393,9 @@ void FAnalyzerHub::OnNewType(const FTypeRegistry::FTypeInfo* TypeInfo)
 
 	uint16 Uid = TypeInfo->Uid;
 	if (Uid >= uint16(TypeToRoute.Num()))
- 	{
+	{
 		TypeToRoute.SetNumZeroed(Uid + 32);
- 	}
+	}
 
 	TypeToRoute[Uid] = uint16(FirstRoute + 1);
 
@@ -1489,7 +1500,7 @@ void FAnalyzerHub::BuildRoutes()
 	};
 
 	int32 AllEventsIndex = FindRoute(FFnv1aHash().Get());
-	auto FixupRoute = [this, &FindRoute, AllEventsIndex] (FRoute* Route) 
+	auto FixupRoute = [this, &FindRoute, AllEventsIndex] (FRoute* Route)
 	{
 		if (Route->ParentHash)
 		{
@@ -1683,7 +1694,8 @@ void FTraceAnalyzer::OnNewTrace(const FOnEventContext& Context)
 
 	// Later traces will have an explicit "SerialSync" trace event to indicate
 	// when there is enough data to establish the correct log serial
-	if ((EventData.GetValue<uint8>("FeatureSet") & 1) == 0)
+	const uint8 FeatureSet = EventData.GetValue<uint8>("FeatureSet");
+	if ((FeatureSet & 1) == 0)
 	{
 		OnSerialSync(Context);
 	}
@@ -1750,7 +1762,7 @@ void FTraceAnalyzer::OnThreadInfoInternal(const FOnEventContext& Context)
 
 	ThreadInfo->SystemId = EventData.GetValue<uint32>("SystemId");
 	ThreadInfo->SortHint = EventData.GetValue<int32>("SortHint");
-	
+
 	FAnsiStringView Name;
 	EventData.GetString("Name", Name);
 	ThreadInfo->Name.SetNumUninitialized(Name.Len() + 1);
@@ -1801,9 +1813,11 @@ public:
 	void				OnEvent(const FEventDataInfo& EventDataInfo);
 	virtual void		OnThreadInfo(const FThreads::FInfo& InThreadInfo) override;
 	void				EnterScope();
-	void				EnterScope(uint64 Timestamp);
+	void				EnterScope(uint64 RelativeTimestamp);
+	void				EnterScopeA(uint64 AbsoluteTimestamp);
 	void				LeaveScope();
-	void				LeaveScope(uint64 Timestamp);
+	void				LeaveScope(uint64 RelativeTimestamp);
+	void				LeaveScopeA(uint64 AbsoluteTimestamp);
 
 private:
 	void				DispatchLeaveScope();
@@ -1910,6 +1924,12 @@ void FAnalysisBridge::EnterScope(uint64 Timestamp)
 	ThreadInfo->ScopeRoutes.Push(~Timestamp);
 }
 
+////////////////////////////////////////////////////////////////////////////////
+void FAnalysisBridge::EnterScopeA(uint64 Timestamp)
+{
+	ThreadInfo->ScopeRoutes.Push(~Timestamp);
+}
+
 ////////////////////////////////////////////////////////////////////////////////
 void FAnalysisBridge::LeaveScope()
 {
@@ -1926,6 +1946,14 @@ void FAnalysisBridge::LeaveScope(uint64 Timestamp)
 	State.Timing.EventTimestamp = 0;
 }
 
+////////////////////////////////////////////////////////////////////////////////
+void FAnalysisBridge::LeaveScopeA(uint64 Timestamp)
+{
+	State.Timing.EventTimestamp = Timestamp;
+	DispatchLeaveScope();
+	State.Timing.EventTimestamp = 0;
+}
+
 ////////////////////////////////////////////////////////////////////////////////
 void FAnalysisBridge::DispatchLeaveScope()
 {
@@ -2180,7 +2208,7 @@ class FProtocol2Stage
 public:
 								FProtocol2Stage(uint32 Version, FTransport* InTransport);
 								~FProtocol2Stage();
-	virtual EStatus 			OnData(FStreamReader& Reader, const FMachineContext& Context) override;
+	virtual EStatus				OnData(FStreamReader& Reader, const FMachineContext& Context) override;
 	virtual void				EnterStage(const FMachineContext& Context) override;
 
 protected:
@@ -2667,7 +2695,7 @@ int32 FProtocol4Stage::OnDataImpl(FStreamReader& Reader, FAnalysisBridge& Bridge
 		{
 			return 1;
 		}
-		
+
 		const auto* HeaderSync = (Protocol4::FEventHeaderSync*)Header;
 		uint32 EventSerial = HeaderSync->SerialLow|(uint32(HeaderSync->SerialHigh) << 16);
 		if (EventSerial != (Serial.Value & Serial.Mask))
@@ -2751,70 +2779,56 @@ int32 FProtocol4Stage::OnDataKnown(
 			const FTypeRegistry::FTypeInfo* TypeInfo = TypeRegistry.Add(Size + 1, 4);
 			Bridge.OnNewType(TypeInfo);
 			Reader.Advance(sizeof(*Size) + *Size);
-			return 1;
 		}
-		break;
+		return 1;
 
 	case Protocol4::EKnownEventUids::EnterScope:
+		Bridge.EnterScope();
+		return 1;
+
+	case Protocol4::EKnownEventUids::LeaveScope:
+		Bridge.LeaveScope();
+		return 1;
+
 	case Protocol4::EKnownEventUids::EnterScope_T:
-		if (Uid > Protocol4::EKnownEventUids::EnterScope)
 		{
 			const uint8* Stamp = Reader.GetPointer(sizeof(uint64) - 1);
 			if (Stamp == nullptr)
 			{
-				break;
+				return 0;
 			}
 
-			uint64 Timestamp = *(uint64*)(Stamp - 1) >> 8;
-			Bridge.EnterScope(Timestamp);
+			const uint64 RelativeTimestamp = *(uint64*)(Stamp - 1) >> 8;
+			Bridge.EnterScope(RelativeTimestamp);
 
 			Reader.Advance(sizeof(uint64) - 1);
 		}
-		else
-		{
-			Bridge.EnterScope();
-		}
 		return 1;
 
-	case Protocol4::EKnownEventUids::LeaveScope:
 	case Protocol4::EKnownEventUids::LeaveScope_T:
-		if (Uid > Protocol4::EKnownEventUids::LeaveScope)
 		{
 			const uint8* Stamp = Reader.GetPointer(sizeof(uint64) - 1);
 			if (Stamp == nullptr)
 			{
-				break;
+				return 0;
 			}
 
-			uint64 Timestamp = *(uint64*)(Stamp - 1) >> 8;
-			Bridge.LeaveScope(Timestamp);
+			const uint64 RelativeTimestamp = *(uint64*)(Stamp - 1) >> 8;
+			Bridge.LeaveScope(RelativeTimestamp);
 
 			Reader.Advance(sizeof(uint64) - 1);
 		}
-		else
-		{
-			Bridge.LeaveScope();
-		}
 		return 1;
-	};
 
-	return 0;
+	default:
+		return 0;
+	};
 }
 
 
 
 // {{{1 protocol-5 -------------------------------------------------------------
 
-////////////////////////////////////////////////////////////////////////////////
-DEFINE_LOG_CATEGORY_STATIC(LogTraceAnalysis, Log, All)
-
-#if 0
-#	define TRACE_ANALYSIS_DEBUG(Format, ...) \
-		do { UE_LOG(LogTraceAnalysis, Log, TEXT(Format), __VA_ARGS__) } while (0)
-#else
-#	define TRACE_ANALYSIS_DEBUG(...)
-#endif
-
 ////////////////////////////////////////////////////////////////////////////////
 class FProtocol5Stage
 	: public FAnalysisMachine::FStage
@@ -2870,7 +2884,7 @@ protected:
 		uint32 Origin;
 	};
 
-	enum ESerial : int32 
+	enum ESerial : int32
 	{
 		Bits		= 24,
 		Mask		= (1 << Bits) - 1,
@@ -2882,7 +2896,6 @@ protected:
 	using EventDescArray	= TArray<FEventDesc>;
 	using EKnownUids		= Protocol5::EKnownEventUids;
 
-
 	virtual EStatus			OnData(FStreamReader& Reader, const FMachineContext& Context) override;
 	EStatus					OnDataNewEvents(const FMachineContext& Context);
 	EStatus					OnDataImportant(const FMachineContext& Context);
@@ -2892,11 +2905,14 @@ protected:
 	int32					ParseEvents(FStreamReader& Reader, EventDescArray& OutEventDescs, const FMachineContext& Context);
 	int32					ParseEventsWithAux(FStreamReader& Reader, EventDescArray& OutEventDescs, const FMachineContext& Context);
 	int32					ParseEvent(FStreamReader& Reader, FEventDesc& OutEventDesc, const FMachineContext& Context);
+	virtual void			SetSizeIfKnownEvent(uint32 Uid, uint32& InOutEventSize);
+	virtual bool			DispatchKnownEvent(const FMachineContext& Context, uint32 Uid, const FEventDesc* Cursor);
 	int32					DispatchEvents(const FMachineContext& Context, const FEventDesc* EventDesc, uint32 Count);
 	int32					DispatchEvents(const FMachineContext& Context, TArray<FEventDescStream>& EventDescHeap);
 	void					DetectSerialGaps(TArray<FEventDescStream>& EventDescHeap);
 	template <typename Callback>
 	void					ForEachSerialGap(const TArray<FEventDescStream>& EventDescHeap, Callback&& InCallback);
+
 	FTypeRegistry			TypeRegistry;
 	FTidPacketTransport&	Transport;
 	EventDescArray			EventDescs;
@@ -3122,7 +3138,7 @@ FProtocol5Stage::EStatus FProtocol5Stage::OnDataNormal(const FMachineContext& Co
 		uint32 NumEventDescs = EventDescs.Num();
 
 		TRACE_ANALYSIS_DEBUG("Thread: %03d Id:%04x", i, Transport.GetThreadId(i));
-		
+
 		// Extract all the events in the stream for this thread
 		FStreamReader* ThreadReader = Transport.GetThreadStream(i);
 
@@ -3253,7 +3269,7 @@ int32 FProtocol5Stage::DispatchEvents(
 		const FEventDesc* EndDesc = StartDesc;
 
 		// DetectSerialGaps() will add a special stream that communicates gaps
-		// in in serial numbers, gaps that will never be resolved. Thread IDs
+		// in serial numbers, gaps that will never be resolved. Thread IDs
 		// are uint16 everywhere else so they will never collide with GapThreadId.
 		if (Stream.ThreadId == FEventDescStream::GapThreadId)
 		{
@@ -3490,12 +3506,7 @@ int32 FProtocol5Stage::ParseEvent(FStreamReader& Reader, FEventDesc& EventDesc,
 		}
 		else
 		{
-			switch (Uid)
-			{
-			case EKnownUids::EnterScope_T:
-			case EKnownUids::LeaveScope_T:
-				EventSize = 7;
-			};
+			SetSizeIfKnownEvent(Uid, EventSize);
 		}
 
 		EventDesc.bHasAux = 0;
@@ -3507,7 +3518,7 @@ int32 FProtocol5Stage::ParseEvent(FStreamReader& Reader, FEventDesc& EventDesc,
 		const FTypeRegistry::FTypeInfo* TypeInfo = TypeRegistry.Get(Uid);
 		if (TypeInfo == nullptr)
 		{
-			Context.Log->Warning(LOCTEXT("UnknownUID","An unknown event UID was encountered."));	
+			Context.Log->Warning(LOCTEXT("UnknownUID","An unknown event UID was encountered."));
 			return 0;
 		}
 
@@ -3638,7 +3649,7 @@ void FProtocol5Stage::DetectSerialGaps(TArray<FEventDescStream>& EventDescHeap)
 			{
 				return false;
 			}
-			
+
 			const FEventDesc& SerialGap = SerialGaps[GapCount];
 
 			if (SerialGap.Serial == Lhs)
@@ -3685,6 +3696,54 @@ void FProtocol5Stage::DetectSerialGaps(TArray<FEventDescStream>& EventDescHeap)
 	}
 }
 
+////////////////////////////////////////////////////////////////////////////////
+void FProtocol5Stage::SetSizeIfKnownEvent(uint32 Uid, uint32& InOutEventSize)
+{
+	switch (Uid)
+	{
+	case EKnownUids::EnterScope_T:
+	case EKnownUids::LeaveScope_T:
+		InOutEventSize = 7;
+	};
+}
+
+////////////////////////////////////////////////////////////////////////////////
+bool FProtocol5Stage::DispatchKnownEvent(const FMachineContext& Context, uint32 Uid, const FEventDesc* Cursor)
+{
+	// Maybe this is a "well-known" event that is handled a little different?
+	switch (Uid)
+	{
+	case EKnownUids::EnterScope:
+		Context.Bridge.EnterScope();
+		return true;
+
+	case EKnownUids::LeaveScope:
+		Context.Bridge.LeaveScope();
+		return true;
+
+	case EKnownUids::EnterScope_T:
+	{
+		uint64 RelativeTimestamp = *(uint64*)(Cursor->Data - 1) >> 8;
+		Context.Bridge.EnterScope(RelativeTimestamp);
+		return true;
+	}
+
+	case EKnownUids::LeaveScope_T:
+	{
+		uint64 RelativeTimestamp = *(uint64*)(Cursor->Data - 1) >> 8;
+		Context.Bridge.LeaveScope(RelativeTimestamp);
+		return true;
+	}
+
+	case EKnownUids::AuxData:
+	case EKnownUids::AuxDataTerminal:
+		return true;
+
+	default:
+		return false;
+	}
+}
+
 ////////////////////////////////////////////////////////////////////////////////
 int32 FProtocol5Stage::DispatchEvents(
 	const FMachineContext& Context,
@@ -3694,43 +3753,23 @@ int32 FProtocol5Stage::DispatchEvents(
 
 	FAuxDataCollector AuxCollector;
 
-	for (const FEventDesc *Cursor = EventDesc, *End = EventDesc + Count; Cursor < End; ++Cursor)
+	for (const FEventDesc* Cursor = EventDesc, *End = EventDesc + Count; Cursor < End; ++Cursor)
 	{
-		// Maybe this is a "well-known" event that is handled a little different?
-		switch (Cursor->Uid)
-		{
-		case EKnownUids::EnterScope_T:
-		{
-			uint64 Timestamp = *(uint64*)(Cursor->Data - 1) >> 8;
-			Context.Bridge.EnterScope(Timestamp);
-			continue;
-		}
+		uint32 Uid = uint32(Cursor->Uid);
 
-		case EKnownUids::LeaveScope_T:
+		if (DispatchKnownEvent(Context, Uid, Cursor))
 		{
-			uint64 Timestamp = *(uint64*)(Cursor->Data - 1) >> 8;
-			Context.Bridge.LeaveScope(Timestamp);
 			continue;
 		}
 
-		case EKnownUids::EnterScope: Context.Bridge.EnterScope(); continue;
-		case EKnownUids::LeaveScope: Context.Bridge.LeaveScope(); continue;
-
-		case EKnownUids::AuxData:
-		case EKnownUids::AuxDataTerminal:
-			continue;
-
-		default:
-			if (!TypeRegistry.IsUidValid(Cursor->Uid))
-			{
-				Context.Log->Warning(LOCTEXT("UnknownUID","An unknown event UID was encountered."));	
-				return -1;
-			}
-			break;
+		if (!TypeRegistry.IsUidValid(Uid))
+		{
+			Context.Log->Warning(LOCTEXT("UnknownUID","An unknown event UID was encountered."));
+			return -1;
 		}
 
 		// It is a normal event.
-		const FTypeRegistry::FTypeInfo* TypeInfo = TypeRegistry.Get(Cursor->Uid);
+		const FTypeRegistry::FTypeInfo* TypeInfo = TypeRegistry.Get(Uid);
 		FEventDataInfo EventDataInfo = {
 			Cursor->Data,
 			*TypeInfo,
@@ -3771,6 +3810,7 @@ int32 FProtocol5Stage::DispatchEvents(
 }
 
 
+
 // {{{1 protocol-6 -------------------------------------------------------------
 
 ////////////////////////////////////////////////////////////////////////////////
@@ -3778,8 +3818,6 @@ class FProtocol6Stage : public FProtocol5Stage
 {
 public:
 							FProtocol6Stage(FTransport* InTransport);
-protected:	
-	using EKnownUids		= Protocol6::EKnownEventUids;
 };
 
 ////////////////////////////////////////////////////////////////////////////////
@@ -3789,6 +3827,78 @@ FProtocol6Stage::FProtocol6Stage(FTransport* InTransport)
 	EventVersion = 6;
 }
 
+
+
+// {{{1 protocol-7 -------------------------------------------------------------
+
+////////////////////////////////////////////////////////////////////////////////
+class FProtocol7Stage : public FProtocol6Stage
+{
+public:
+							FProtocol7Stage(FTransport* InTransport);
+
+	virtual void			SetSizeIfKnownEvent(uint32 Uid, uint32& InOutEventSize) override;
+	virtual bool			DispatchKnownEvent(const FMachineContext& Context, uint32 Uid, const FEventDesc* Cursor) override;
+
+protected:
+	using EKnownUids		= Protocol7::EKnownEventUids;
+};
+
+////////////////////////////////////////////////////////////////////////////////
+FProtocol7Stage::FProtocol7Stage(FTransport* InTransport)
+	: FProtocol6Stage(InTransport)
+{
+	EventVersion = 7;
+}
+
+////////////////////////////////////////////////////////////////////////////////
+void FProtocol7Stage::SetSizeIfKnownEvent(uint32 Uid, uint32& InOutEventSize)
+{
+	switch (Uid)
+	{
+	case EKnownUids::EnterScope_TA:
+	case EKnownUids::LeaveScope_TA:
+		InOutEventSize = 8;
+	};
+}
+
+////////////////////////////////////////////////////////////////////////////////
+bool FProtocol7Stage::DispatchKnownEvent(const FMachineContext& Context, uint32 Uid, const FEventDesc* Cursor)
+{
+	// Maybe this is a "well-known" event that is handled a little different?
+	switch (Uid)
+	{
+	case EKnownUids::EnterScope:
+		Context.Bridge.EnterScope();
+		return true;
+
+	case EKnownUids::LeaveScope:
+		Context.Bridge.LeaveScope();
+		return true;
+
+	case EKnownUids::EnterScope_TA:
+	{
+		uint64 AbsoluteTimestamp = *(uint64*)Cursor->Data;
+		Context.Bridge.EnterScopeA(AbsoluteTimestamp);
+		return true;
+	}
+
+	case EKnownUids::LeaveScope_TA:
+	{
+		uint64 AbsoluteTimestamp = *(uint64*)Cursor->Data;
+		Context.Bridge.LeaveScopeA(AbsoluteTimestamp);
+		return true;
+	}
+
+	case EKnownUids::AuxData:
+	case EKnownUids::AuxDataTerminal:
+		return true;
+
+	default:
+		return false;
+	}
+}
+
 // {{{1 est.-transport ---------------------------------------------------------
 
 ////////////////////////////////////////////////////////////////////////////////
@@ -3815,8 +3925,14 @@ FEstablishTransportStage::EStatus FEstablishTransportStage::OnData(
 		return EStatus::NotEnoughData;
 	}
 
+	uint32 TransportVersion = Header->TransportVersion;
+	TRACE_ANALYSIS_DEBUG("TransportVersion: %u", TransportVersion);
+
+	uint32 ProtocolVersion = Header->ProtocolVersion;
+	TRACE_ANALYSIS_DEBUG("ProtocolVersion: %u", ProtocolVersion);
+
 	FTransport* Transport = nullptr;
-	switch (Header->TransportVersion)
+	switch (TransportVersion)
 	{
 	case ETransport::Raw:			Transport = new FTransport(); break;
 	case ETransport::Packet:		Transport = new FPacketTransport(); break;
@@ -3826,12 +3942,11 @@ FEstablishTransportStage::EStatus FEstablishTransportStage::OnData(
 		{
 			Context.Log->Error(FText::Format(
 				LOCTEXT("InvalidTransportVersion", "Unknown transport version: {0}. You may need to recompile this application"),
-				FText::AsNumber(Header->TransportVersion)));
+				FText::AsNumber(TransportVersion)));
 			return EStatus::Error;
 		}
 	}
 
-	uint32 ProtocolVersion = Header->ProtocolVersion;
 	switch (ProtocolVersion)
 	{
 	case Protocol0::EProtocol::Id:
@@ -3861,6 +3976,11 @@ FEstablishTransportStage::EStatus FEstablishTransportStage::OnData(
 		Context.Machine.Transition();
 		break;
 
+	case Protocol7::EProtocol::Id:
+		Context.Machine.QueueStage<FProtocol7Stage>(Transport);
+		Context.Machine.Transition();
+		break;
+
 	default:
 		{
 			Context.Log->Error(FText::Format(
@@ -4052,6 +4172,7 @@ bool FAnalysisEngine::OnData(FStreamReader& Reader)
 } // namespace Trace
 } // namespace UE
 
+#undef TRACE_ANALYSIS_DEBUG
 #undef LOCTEXT_NAMESPACE
 
 /* vim: set foldlevel=1 : */
