[Insights] Fixed trace analysis for scoped trace events (adjusts time to be relative to the base timestamp).

diff --git a/Engine/Source/Developer/TraceAnalysis/Private/Analysis/Engine.cpp b/Engine/Source/Developer/TraceAnalysis/Private/Analysis/Engine.cpp
--- a/Engine/Source/Developer/TraceAnalysis/Private/Analysis/Engine.cpp
+++ b/Engine/Source/Developer/TraceAnalysis/Private/Analysis/Engine.cpp
@@ -1927,6 +1927,7 @@ void FAnalysisBridge::EnterScope(uint64 Timestamp)
 ////////////////////////////////////////////////////////////////////////////////
 void FAnalysisBridge::EnterScopeA(uint64 Timestamp)
 {
+	Timestamp -= State.Timing.BaseTimestamp;
 	ThreadInfo->ScopeRoutes.Push(~Timestamp);
 }
 
@@ -1949,6 +1950,7 @@ void FAnalysisBridge::LeaveScope(uint64 Timestamp)
 ////////////////////////////////////////////////////////////////////////////////
 void FAnalysisBridge::LeaveScopeA(uint64 Timestamp)
 {
+	Timestamp -= State.Timing.BaseTimestamp;
 	State.Timing.EventTimestamp = Timestamp;
 	DispatchLeaveScope();
 	State.Timing.EventTimestamp = 0;
