[Insights] Fixed time for scoped events to be relative to base timestamp at runtime. This allows packing on 8 bytes (as in trace protocol 5). Added EnterScope_TB / LeaveScope_TB (timestamps relative to BaseTimestamp) events in addition to EnterScope_TA / LeaveScope_TA (absolute timestamps).

diff --git a/Engine/Source/Developer/TraceAnalysis/Private/Analysis/Engine.cpp b/Engine/Source/Developer/TraceAnalysis/Private/Analysis/Engine.cpp
--- a/Engine/Source/Developer/TraceAnalysis/Private/Analysis/Engine.cpp
+++ b/Engine/Source/Developer/TraceAnalysis/Private/Analysis/Engine.cpp
@@ -1815,9 +1815,11 @@ public:
 	void				EnterScope();
 	void				EnterScope(uint64 RelativeTimestamp);
 	void				EnterScopeA(uint64 AbsoluteTimestamp);
+	void				EnterScopeB(uint64 BaseRelativeTimestamp);
 	void				LeaveScope();
 	void				LeaveScope(uint64 RelativeTimestamp);
 	void				LeaveScopeA(uint64 AbsoluteTimestamp);
+	void				LeaveScopeB(uint64 BaseRelativeTimestamp);
 
 private:
 	void				DispatchLeaveScope();
@@ -1931,6 +1933,12 @@ void FAnalysisBridge::EnterScopeA(uint64 Timestamp)
 	ThreadInfo->ScopeRoutes.Push(~Timestamp);
 }
 
+////////////////////////////////////////////////////////////////////////////////
+void FAnalysisBridge::EnterScopeB(uint64 Timestamp)
+{
+	ThreadInfo->ScopeRoutes.Push(~Timestamp);
+}
+
 ////////////////////////////////////////////////////////////////////////////////
 void FAnalysisBridge::LeaveScope()
 {
@@ -1956,6 +1964,14 @@ void FAnalysisBridge::LeaveScopeA(uint64 Timestamp)
 	State.Timing.EventTimestamp = 0;
 }
 
+////////////////////////////////////////////////////////////////////////////////
+void FAnalysisBridge::LeaveScopeB(uint64 Timestamp)
+{
+	State.Timing.EventTimestamp = Timestamp;
+	DispatchLeaveScope();
+	State.Timing.EventTimestamp = 0;
+}
+
 ////////////////////////////////////////////////////////////////////////////////
 void FAnalysisBridge::DispatchLeaveScope()
 {
@@ -3706,6 +3722,7 @@ void FProtocol5Stage::SetSizeIfKnownEvent(uint32 Uid, uint32& InOutEventSize)
 	case EKnownUids::EnterScope_T:
 	case EKnownUids::LeaveScope_T:
 		InOutEventSize = 7;
+		break;
 	};
 }
 
@@ -3861,6 +3878,12 @@ void FProtocol7Stage::SetSizeIfKnownEvent(uint32 Uid, uint32& InOutEventSize)
 	case EKnownUids::EnterScope_TA:
 	case EKnownUids::LeaveScope_TA:
 		InOutEventSize = 8;
+		break;
+
+	case EKnownUids::EnterScope_TB:
+	case EKnownUids::LeaveScope_TB:
+		InOutEventSize = 7;
+		break;
 	};
 }
 
@@ -3892,6 +3915,20 @@ bool FProtocol7Stage::DispatchKnownEvent(const FMachineContext& Context, uint32
 		return true;
 	}
 
+	case EKnownUids::EnterScope_TB:
+	{
+		uint64 BaseRelativeTimestamp = *(uint64*)(Cursor->Data - 1) >> 8;
+		Context.Bridge.EnterScopeB(BaseRelativeTimestamp);
+		return true;
+	}
+
+	case EKnownUids::LeaveScope_TB:
+	{
+		uint64 BaseRelativeTimestamp = *(uint64*)(Cursor->Data - 1) >> 8;
+		Context.Bridge.LeaveScopeB(BaseRelativeTimestamp);
+		return true;
+	}
+
 	case EKnownUids::AuxData:
 	case EKnownUids::AuxDataTerminal:
 		return true;
diff --git a/Engine/Source/Runtime/TraceLog/Private/Trace/Writer.cpp b/Engine/Source/Runtime/TraceLog/Private/Trace/Writer.cpp
--- a/Engine/Source/Runtime/TraceLog/Private/Trace/Writer.cpp
+++ b/Engine/Source/Runtime/TraceLog/Private/Trace/Writer.cpp
@@ -76,7 +76,7 @@ UE_TRACE_EVENT_END()
 ////////////////////////////////////////////////////////////////////////////////
 static volatile bool			GInitialized;		// = false;
 FStatistics						GTraceStatistics;	// = {};
-uint64							GStartCycle;		// = 0;
+TRACELOG_API uint64				GStartCycle;		// = 0;
 TRACELOG_API uint32 volatile	GLogSerial;			// = 0;
 // Counter of calls to Writer_WorkerUpdate to enable regular flushing of output buffers
 static uint32					GUpdateCounter;		// = 0;
diff --git a/Engine/Source/Runtime/TraceLog/Public/Trace/Detail/LogScope.inl b/Engine/Source/Runtime/TraceLog/Public/Trace/Detail/LogScope.inl
--- a/Engine/Source/Runtime/TraceLog/Public/Trace/Detail/LogScope.inl
+++ b/Engine/Source/Runtime/TraceLog/Public/Trace/Detail/LogScope.inl
@@ -18,6 +18,7 @@ namespace Trace {
 namespace Private {
 
 ////////////////////////////////////////////////////////////////////////////////
+extern TRACELOG_API uint64			GStartCycle;
 extern TRACELOG_API uint32 volatile	GLogSerial;
 
 ////////////////////////////////////////////////////////////////////////////////
@@ -165,18 +166,16 @@ inline FScopedStampedLogScope::~FScopedStampedLogScope()
 		return;
 	}
 
-	uint8 LeaveUid = uint8(EKnownEventUids::LeaveScope_TA << EKnownEventUids::_UidShift);
-	uint64 Stamp = TimeGetTimestamp();
+	uint64 Stamp = TimeGetTimestamp() - GStartCycle;
 
 	FWriteBuffer* Buffer = Writer_GetBuffer();
-	if (UNLIKELY(int32((uint8*)Buffer - Buffer->Cursor) < int32(sizeof(LeaveUid)) + int32(sizeof(Stamp))))
+	if (UNLIKELY(int32((uint8*)Buffer - Buffer->Cursor) < int32(sizeof(Stamp))))
 	{
 		Buffer = Writer_NextBuffer();
 	}
 
-	Buffer->Cursor[0] = LeaveUid;
-	Buffer->Cursor += sizeof(LeaveUid);
-
+	Stamp <<= 8;
+	Stamp += uint8(EKnownEventUids::LeaveScope_TB) << EKnownEventUids::_UidShift;
 	memcpy((uint64*)(Buffer->Cursor), &Stamp, sizeof(Stamp));
 	Buffer->Cursor += sizeof(Stamp);
 
@@ -224,19 +223,17 @@ FORCENOINLINE auto FLogScope::ScopedEnter()
 template <class EventType>
 FORCENOINLINE auto FLogScope::ScopedStampedEnter()
 {
-	uint8 EnterUid = uint8(EKnownEventUids::EnterScope_TA << EKnownEventUids::_UidShift);
 	uint64 Stamp;
 
 	FWriteBuffer* Buffer = Writer_GetBuffer();
-	if (UNLIKELY(int32((uint8*)Buffer - Buffer->Cursor) < int32(sizeof(EnterUid)) + int32(sizeof(Stamp))))
+	if (UNLIKELY(int32((uint8*)Buffer - Buffer->Cursor) < int32(sizeof(Stamp))))
 	{
 		Buffer = Writer_NextBuffer();
 	}
 
-	Buffer->Cursor[0] = EnterUid;
-	Buffer->Cursor += sizeof(EnterUid);
-
-	Stamp = TimeGetTimestamp();
+	Stamp = TimeGetTimestamp() - GStartCycle;
+	Stamp <<= 8;
+	Stamp += uint8(EKnownEventUids::EnterScope_TB) << EKnownEventUids::_UidShift;
 	memcpy((uint64*)(Buffer->Cursor), &Stamp, sizeof(Stamp));
 	Buffer->Cursor += sizeof(Stamp);
 
diff --git a/Engine/Source/Runtime/TraceLog/Public/Trace/Detail/Protocols/Protocol7.h b/Engine/Source/Runtime/TraceLog/Public/Trace/Detail/Protocols/Protocol7.h
--- a/Engine/Source/Runtime/TraceLog/Public/Trace/Detail/Protocols/Protocol7.h
+++ b/Engine/Source/Runtime/TraceLog/Public/Trace/Detail/Protocols/Protocol7.h
@@ -39,16 +39,16 @@ struct EKnownEventUids
 		AuxDataTerminal,		// same as Protocol5
 		EnterScope,				// same as Protocol5
 		LeaveScope,				// same as Protocol5
-		EnterScope_TA,			// new in Protocol7
-		LeaveScope_TA,			// new in Protocol7
+		EnterScope_TA,			// new in Protocol7, absolute timestamps
+		LeaveScope_TA,			// new in Protocol7, absolute timestamps
+		EnterScope_TB,			// new in Protocol7, timestamps relative to BaseTimestamp
+		LeaveScope_TB,			// new in Protocol7, timestamps relative to BaseTimestamp
 		_Unused1,
 		_Unused2,
 		_Unused3,
 		_Unused4,
 		_Unused5,
 		_Unused6,
-		_Unused7,
-		_Unused8,
 		_WellKnownNum,			// same as Protocol5
 	};
 	static const uint16 User = _WellKnownNum;
