Add user facing error message if transport layer error is detected.

diff --git a/Engine/Source/Developer/TraceAnalysis/Private/Analysis/Engine.cpp b/Engine/Source/Developer/TraceAnalysis/Private/Analysis/Engine.cpp
--- a/Engine/Source/Developer/TraceAnalysis/Private/Analysis/Engine.cpp
+++ b/Engine/Source/Developer/TraceAnalysis/Private/Analysis/Engine.cpp
@@ -2299,7 +2299,12 @@ FProtocol2Stage::EStatus FProtocol2Stage::OnData(
 {
 	auto* InnerTransport = (FTidPacketTransport*)Transport;
 	InnerTransport->SetReader(Reader);
-	InnerTransport->Update();
+	const FTidPacketTransport::ETransportResult Result = InnerTransport->Update();
+	if (Result == FTidPacketTransport::ETransportResult::Error)
+	{
+		Context.Log->Error(LOCTEXT("TransportError", "An error was detected in the transport layer, most likely due to a corrupt trace file. See log for details."));
+		return EStatus::Error;
+	}
 
 	struct FRotaItem
 	{
@@ -2954,7 +2959,12 @@ FProtocol5Stage::EStatus FProtocol5Stage::OnData(
 	const FMachineContext& Context)
 {
 	Transport.SetReader(Reader);
-	Transport.Update();
+	const FTidPacketTransport::ETransportResult Result = Transport.Update();
+	if (Result == FTidPacketTransport::ETransportResult::Error)
+	{
+		Context.Log->Error(LOCTEXT("TransportError", "An error was detected in the transport layer, most likely due to a corrupt trace file. See log for details."));
+		return EStatus::Error;
+	}
 
 	// New-events. They must be processed before anything else otherwise events
 	// can not be interpreted.
diff --git a/Engine/Source/Developer/TraceAnalysis/Private/Analysis/Transport/TidPacketTransport.cpp b/Engine/Source/Developer/TraceAnalysis/Private/Analysis/Transport/TidPacketTransport.cpp
--- a/Engine/Source/Developer/TraceAnalysis/Private/Analysis/Transport/TidPacketTransport.cpp
+++ b/Engine/Source/Developer/TraceAnalysis/Private/Analysis/Transport/TidPacketTransport.cpp
@@ -3,6 +3,7 @@
 #include "TidPacketTransport.h"
 #include "Algo/BinarySearch.h"
 #include "HAL/UnrealMemory.h"
+#include "CoreGlobals.h"
 
 ////////////////////////////////////////////////////////////////////////////////
 namespace UE {
@@ -21,19 +22,19 @@ namespace UE {
 namespace Trace {
 
 ////////////////////////////////////////////////////////////////////////////////
-bool FTidPacketTransport::ReadPacket()
+FTidPacketTransport::EReadPacketResult FTidPacketTransport::ReadPacket()
 {
 	using namespace UE::Trace::Private;
 
 	const auto* PacketBase = GetPointer<FTidPacketBase>();
 	if (PacketBase == nullptr)
 	{
-		return false;
+		return EReadPacketResult::NeedMoreData;
 	}
 
 	if (GetPointer<uint8>(PacketBase->PacketSize) == nullptr)
 	{
-		return false;
+		return EReadPacketResult::NeedMoreData;
 	}
 
 	FTransport::Advance(PacketBase->PacketSize);
@@ -43,15 +44,15 @@ bool FTidPacketTransport::ReadPacket()
 	if (ThreadId == ETransportTid::Sync)
 	{
 		++Synced;
-		return false;	// Do not read any more packets. Gives consumers a
-						// chance to sample the world at each known sync point.
+		return EReadPacketResult::NeedMoreData;	// Do not read any more packets. Gives consumers a
+													// chance to sample the world at each known sync point.
 	}
 
 	bool bIsPartial = !!(PacketBase->ThreadId & FTidPacketBase::PartialMarker);
 	FThreadStream* Thread = FindOrAddThread(ThreadId, !bIsPartial);
 	if (Thread == nullptr)
 	{
-		return true;
+		return EReadPacketResult::Continue;
 	}
 
 	uint32 DataSize = PacketBase->PacketSize - sizeof(FTidPacketBase);
@@ -62,14 +63,18 @@ bool FTidPacketTransport::ReadPacket()
 		uint8* Dest = Thread->Buffer.Append(DecodedSize);
 		DataSize -= sizeof(DecodedSize);
 		int32 ResultSize = UE::Trace::Private::Decode(Packet->Data, DataSize, Dest, DecodedSize);
-		check(int32(DecodedSize) == ResultSize);
+		if (int32(DecodedSize) != ResultSize)
+		{
+			UE_LOG(LogCore, Error, TEXT("Unable to decompress packet, expected %d bytes but decoded %d bytes."), DecodedSize, ResultSize);
+			return EReadPacketResult::ReadError;
+		}
 	}
 	else
 	{
 		Thread->Buffer.Append((uint8*)(PacketBase + 1), DataSize);
 	}
 
-	return true;
+	return EReadPacketResult::Continue;
 }
 
 ////////////////////////////////////////////////////////////////////////////////
@@ -98,14 +103,20 @@ FTidPacketTransport::FThreadStream* FTidPacketTransport::FindOrAddThread(
 }
 
 ////////////////////////////////////////////////////////////////////////////////
-void FTidPacketTransport::Update()
+FTidPacketTransport::ETransportResult FTidPacketTransport::Update()
 {
-	while (ReadPacket());
+	EReadPacketResult Result;
+	do
+	{
+		Result = ReadPacket();
+	} while (Result == EReadPacketResult::Continue);
 
 	Threads.RemoveAll([] (const FThreadStream& Thread)
 	{
 		return (Thread.ThreadId <= ETransportTid::Importants) ? false : Thread.Buffer.IsEmpty();
 	});
+
+	return Result == EReadPacketResult::ReadError ? ETransportResult::Error : ETransportResult::Ok;
 }
 
 ////////////////////////////////////////////////////////////////////////////////
diff --git a/Engine/Source/Developer/TraceAnalysis/Private/Analysis/Transport/TidPacketTransport.h b/Engine/Source/Developer/TraceAnalysis/Private/Analysis/Transport/TidPacketTransport.h
--- a/Engine/Source/Developer/TraceAnalysis/Private/Analysis/Transport/TidPacketTransport.h
+++ b/Engine/Source/Developer/TraceAnalysis/Private/Analysis/Transport/TidPacketTransport.h
@@ -18,21 +18,35 @@ class FTidPacketTransport
 {
 public:
 	typedef UPTRINT ThreadIter;
+	
+	enum class ETransportResult : uint8
+	{
+		Ok,
+		Error
+	};
 
-	void					Update();
+	ETransportResult		Update();
 	uint32					GetThreadCount() const;
 	FStreamReader*			GetThreadStream(uint32 Index);
 	uint32					GetThreadId(uint32 Index) const;
 	uint32					GetSyncCount() const;
 
 private:
+	
 	struct FThreadStream
 	{
 		FStreamBuffer		Buffer;
 		uint32				ThreadId;
 	};
 
-	bool					ReadPacket();
+	enum class EReadPacketResult : uint8
+	{
+		NeedMoreData,
+		Continue,
+		ReadError
+	};
+
+	EReadPacketResult		ReadPacket();
 	FThreadStream*			FindOrAddThread(uint32 ThreadId, bool bAddIfNotFound);
 	TArray<FThreadStream>	Threads = {
 								{ {}, ETransportTid::Events },
