[Insights] Fixed StoreClient's read trace (FStoreCborClient::ReadTrace) incorectly using socket's local_endpoint insted of remote_endpoint.

diff --git a/Engine/Source/Developer/TraceAnalysis/Private/Store/StoreClient.cpp b/Engine/Source/Developer/TraceAnalysis/Private/Store/StoreClient.cpp
--- a/Engine/Source/Developer/TraceAnalysis/Private/Store/StoreClient.cpp
+++ b/Engine/Source/Developer/TraceAnalysis/Private/Store/StoreClient.cpp
@@ -323,7 +323,7 @@ FTraceDataStream* FStoreCborClient::ReadTrace(uint32 Id)
 		return nullptr;
 	}
 
-	asio::ip::address ServerAddr = Socket.local_endpoint().address();
+	asio::ip::address ServerAddr = Socket.remote_endpoint().address();
 	asio::ip::tcp::endpoint Endpoint(ServerAddr, uint16(SenderPort));
 
 	asio::error_code ErrorCode;
