Fix missing dependency
#rnx

diff --git a/Engine/Source/Developer/TraceAnalysis/TraceAnalysis.Build.cs b/Engine/Source/Developer/TraceAnalysis/TraceAnalysis.Build.cs
--- a/Engine/Source/Developer/TraceAnalysis/TraceAnalysis.Build.cs
+++ b/Engine/Source/Developer/TraceAnalysis/TraceAnalysis.Build.cs
@@ -12,6 +12,7 @@ public class TraceAnalysis : ModuleRules
 		PrivateDependencyModuleNames.Add("Asio");
 		PrivateDependencyModuleNames.Add("Cbor");
 		PrivateDependencyModuleNames.Add("Core");
+		PrivateDependencyModuleNames.Add("MessageLog");
 		PrivateDependencyModuleNames.Add("Sockets");
 		PrivateDependencyModuleNames.Add("TraceLog");
 
