Remove usage of MessageLog in analysis.

Usage of MessageLog in non-ui threads is tricky and need to be carefully managed to avoid flushing memory on the wrong thread. With this change direct usage of MessageLog in the analysis engine is removed and replaced with a callback system which allows the user of TraceAnalysis to implement display of the messages themselves. In the case of InsightsManager this means queuing up the messages in analysis session and regularly polled from the UI thread.

diff --git a/Engine/Source/Developer/TraceAnalysis/Private/Analysis/Context.cpp b/Engine/Source/Developer/TraceAnalysis/Private/Analysis/Context.cpp
--- a/Engine/Source/Developer/TraceAnalysis/Private/Analysis/Context.cpp
+++ b/Engine/Source/Developer/TraceAnalysis/Private/Analysis/Context.cpp
@@ -12,13 +12,19 @@ void FAnalysisContext::AddAnalyzer(IAnalyzer& Analyzer)
 	Analyzers.Add(&Analyzer);
 }
 
+////////////////////////////////////////////////////////////////////////////////
+void FAnalysisContext::SetMessageDelegate(FMessageDelegate Delegate)
+{
+	OnMessage = Delegate;
+}
+
 ////////////////////////////////////////////////////////////////////////////////
 FAnalysisProcessor FAnalysisContext::Process(IInDataStream& DataStream)
 {
 	FAnalysisProcessor Processor;
 	if (Analyzers.Num() > 0)
 	{
-		Processor.Impl = new FAnalysisProcessor::FImpl(DataStream, MoveTemp(Analyzers));
+		Processor.Impl = new FAnalysisProcessor::FImpl(DataStream, MoveTemp(Analyzers), MoveTemp(OnMessage));
 	}
 	return MoveTemp(Processor);
 }
diff --git a/Engine/Source/Developer/TraceAnalysis/Private/Analysis/Engine.cpp b/Engine/Source/Developer/TraceAnalysis/Private/Analysis/Engine.cpp
--- a/Engine/Source/Developer/TraceAnalysis/Private/Analysis/Engine.cpp
+++ b/Engine/Source/Developer/TraceAnalysis/Private/Analysis/Engine.cpp
@@ -10,9 +10,9 @@
 #include "HAL/UnrealMemory.h"
 #include "Internationalization/Internationalization.h"
 #include "Logging/LogMacros.h"
-#include "Logging/MessageLog.h"
 #include "StreamReader.h"
 #include "Templates/UnrealTemplate.h"
+#include "Trace/Analysis.h"
 #include "Trace/Analyzer.h"
 #include "Trace/Detail/Protocol.h"
 #include "Trace/Detail/Transport.h"
@@ -2025,7 +2025,20 @@ public:
 	{
 		FAnalysisMachine&	Machine;
 		FAnalysisBridge&	Bridge;
-		FMessageLog*		Log;
+		FMessageDelegate&	OnMessage;
+
+		inline void EmitMessage(EAnalysisMessageSeverity Severity, FStringView Message) const
+		{
+			const bool _ = OnMessage.ExecuteIfBound(Severity, Message);
+		}
+
+		template<typename FormatType, typename... Types>
+		inline void EmitMessagef(EAnalysisMessageSeverity Severity, const FormatType& Format, Types... Args) const
+		{
+			TStringBuilder<64> FormattedMessage;
+			FormattedMessage.Appendf(Format, Forward<Types>(Args)...);
+			EmitMessage(Severity, FormattedMessage.ToView());
+		}
 	};
 
 	class FStage
@@ -2040,7 +2053,7 @@ public:
 		virtual void		ExitStage(const FMachineContext& Context) {};
 	};
 
-							FAnalysisMachine(FAnalysisBridge& InBridge, FMessageLog* Log);
+							FAnalysisMachine(FAnalysisBridge& InBridge, FMessageDelegate&& InMessage);
 							~FAnalysisMachine();
 	EStatus					OnData(FStreamReader& Reader);
 	void					Transition();
@@ -2053,13 +2066,13 @@ private:
 	FStage*					ActiveStage = nullptr;
 	TArray<FStage*>			StageQueue;
 	TArray<FStage*>			DeadStages;
-	FMessageLog*			Log;
+	FMessageDelegate		OnMessage;
 };
 
 ////////////////////////////////////////////////////////////////////////////////
-FAnalysisMachine::FAnalysisMachine(FAnalysisBridge& InBridge, FMessageLog* InLog)
+FAnalysisMachine::FAnalysisMachine(FAnalysisBridge& InBridge, FMessageDelegate&& InMessage)
 : Bridge(InBridge)
-, Log(InLog)
+, OnMessage(InMessage)
 {
 }
 
@@ -2093,7 +2106,7 @@ void FAnalysisMachine::Transition()
 {
 	if (ActiveStage != nullptr)
 	{
-		const FMachineContext Context = { *this, Bridge, Log };
+		const FMachineContext Context = { *this, Bridge, OnMessage };
 		ActiveStage->ExitStage(Context);
 
 		DeadStages.Add(ActiveStage);
@@ -2103,7 +2116,7 @@ void FAnalysisMachine::Transition()
 
 	if (ActiveStage != nullptr)
 	{
-		const FMachineContext Context = { *this, Bridge, Log };
+		const FMachineContext Context = { *this, Bridge, OnMessage };
 		ActiveStage->EnterStage(Context);
 	}
 }
@@ -2111,7 +2124,7 @@ void FAnalysisMachine::Transition()
 ////////////////////////////////////////////////////////////////////////////////
 FAnalysisMachine::EStatus FAnalysisMachine::OnData(FStreamReader& Reader)
 {
-	const FMachineContext Context = { *this, Bridge, Log };
+	const FMachineContext Context = { *this, Bridge, OnMessage };
 	EStatus Ret;
 	do
 	{
@@ -2302,7 +2315,10 @@ FProtocol2Stage::EStatus FProtocol2Stage::OnData(
 	const FTidPacketTransport::ETransportResult Result = InnerTransport->Update();
 	if (Result == FTidPacketTransport::ETransportResult::Error)
 	{
-		Context.Log->Error(LOCTEXT("TransportError", "An error was detected in the transport layer, most likely due to a corrupt trace file. See log for details."));
+		Context.EmitMessage(
+			EAnalysisMessageSeverity::Error,
+			TEXT("An error was detected in the transport layer, most likely due to a corrupt trace file. See log for details.")
+		);
 		return EStatus::Error;
 	}
 
@@ -2962,7 +2978,10 @@ FProtocol5Stage::EStatus FProtocol5Stage::OnData(
 	const FTidPacketTransport::ETransportResult Result = Transport.Update();
 	if (Result == FTidPacketTransport::ETransportResult::Error)
 	{
-		Context.Log->Error(LOCTEXT("TransportError", "An error was detected in the transport layer, most likely due to a corrupt trace file. See log for details."));
+		Context.EmitMessage(
+			EAnalysisMessageSeverity::Error,
+			TEXT("An error was detected in the transport layer, most likely due to a corrupt trace file. See log for details.")
+		);
 		return EStatus::Error;
 	}
 
@@ -3105,7 +3124,7 @@ int32 FProtocol5Stage::ParseImportantEvents(FStreamReader& Reader, EventDescArra
 		const FTypeRegistry::FTypeInfo* TypeInfo = TypeRegistry.Get(Uid);
 		if (TypeInfo == nullptr)
 		{
-			Context.Log->Warning(LOCTEXT("UnknownUID","An unknown event UID was encountered."));
+			Context.EmitMessagef(EAnalysisMessageSeverity::Warning, TEXT("An unknown event UID (%d) was encountered."), Uid);
 			return 1;
 		}
 
@@ -3139,7 +3158,7 @@ int32 FProtocol5Stage::ParseImportantEvents(FStreamReader& Reader, EventDescArra
 
 			if (Cursor[0] != uint8(EKnownUids::AuxDataTerminal))
 			{
-				Context.Log->Warning(LOCTEXT("AuxDataTerminalExpected","Expected an aux data terminal in the stream."));
+				Context.EmitMessage(EAnalysisMessageSeverity::Warning, TEXT("Expected an aux data terminal in the stream."));
 				return -1;
 			}
 		}
@@ -3172,9 +3191,14 @@ FProtocol5Stage::EStatus FProtocol5Stage::OnDataNormal(const FMachineContext& Co
 
 		// Stop analysis if this thread has accumulated too much data.
 		// This can happen on corrupted traces (ex. with serial sync events missing or out of order).
-		if (ThreadReader->GetRemaining() > 512 * 1024 * 1024)
+		constexpr int32 MaxAccumulatedBytes = 512 * 1024 * 1024;
+		if (ThreadReader->GetRemaining() > MaxAccumulatedBytes)
 		{
-			Context.Log->Error(LOCTEXT("TooMuchData", "Analysis accumulated too much data before being able to continue analysing the stream."));
+			Context.EmitMessagef(
+				EAnalysisMessageSeverity::Error,
+				TEXT("Analysis accumulated too much data (%d MiB) before being able to continue analysing the stream."),
+				ThreadReader->GetRemaining() / (1024*1024)
+			);
 			return EStatus::Error;
 		}
 
@@ -3546,7 +3570,7 @@ int32 FProtocol5Stage::ParseEvent(FStreamReader& Reader, FEventDesc& EventDesc,
 		const FTypeRegistry::FTypeInfo* TypeInfo = TypeRegistry.Get(Uid);
 		if (TypeInfo == nullptr)
 		{
-			Context.Log->Warning(LOCTEXT("UnknownUID","An unknown event UID was encountered."));
+			Context.EmitMessagef(EAnalysisMessageSeverity::Warning, TEXT("An unknown event UID (%u) was encountered."), Uid);
 			return 0;
 		}
 
@@ -3793,7 +3817,7 @@ int32 FProtocol5Stage::DispatchEvents(
 
 		if (!TypeRegistry.IsUidValid(Uid))
 		{
-			Context.Log->Warning(LOCTEXT("UnknownUID","An unknown event UID was encountered."));
+			Context.EmitMessagef(EAnalysisMessageSeverity::Warning, TEXT("An unknown event UID (%u) was encountered."), Uid);
 			return -1;
 		}
 
@@ -3989,9 +4013,11 @@ FEstablishTransportStage::EStatus FEstablishTransportStage::OnData(
 	case ETransport::TidPacketSync:	Transport = new FTidPacketTransportSync(); break;
 	default:
 		{
-			Context.Log->Error(FText::Format(
-				LOCTEXT("InvalidTransportVersion", "Unknown transport version: {0}. You may need to recompile this application"),
-				FText::AsNumber(TransportVersion)));
+			Context.EmitMessagef(
+				EAnalysisMessageSeverity::Error,
+				TEXT("Unknown transport version: %u. You may need to recompile this application"),
+				TransportVersion
+			);
 			return EStatus::Error;
 		}
 	}
@@ -4032,9 +4058,11 @@ FEstablishTransportStage::EStatus FEstablishTransportStage::OnData(
 
 	default:
 		{
-			Context.Log->Error(FText::Format(
-					LOCTEXT("UnknownProtocolVersion", "Unknown protocol version: {0}. You may need to recompile this application."),
-					FText::AsNumber(ProtocolVersion)));
+			Context.EmitMessagef(
+				EAnalysisMessageSeverity::Error,
+				TEXT("Unknown protocol version: %u. You may need to recompile this application"),
+				ProtocolVersion
+			);
 			return EStatus::Error;
 		}
 	}
@@ -4107,7 +4135,7 @@ FMagicStage::EStatus FMagicStage::OnData(
 	if (Magic == 'ECRT' || Magic == '2CRT')
 	{
 		// Source is big-endian which we don't currently support
-		Context.Log->Error(LOCTEXT("BigEndian", "Big endian traces are currently not supported."));
+		Context.EmitMessage(EAnalysisMessageSeverity::Error, TEXT("Big endian traces are currently not supported."));
 		return EStatus::Error;
 	}
 
@@ -4136,7 +4164,7 @@ FMagicStage::EStatus FMagicStage::OnData(
 		return EStatus::Continue;
 	}
 
-	Context.Log->Error(LOCTEXT("IncorrectMagic", "The file or stream was not recognized as trace stream."));
+	Context.EmitMessage(EAnalysisMessageSeverity::Error, TEXT("The file or stream was not recognized as trace stream."));
 	return EStatus::Error;
 }
 
@@ -4148,7 +4176,7 @@ FMagicStage::EStatus FMagicStage::OnData(
 class FAnalysisEngine::FImpl
 {
 public:
-					FImpl(TArray<IAnalyzer*>&& Analyzers, FMessageLog* InLog);
+						FImpl(TArray<IAnalyzer*>&& Analyzers, FMessageDelegate&& InMessage);
 	void				Begin();
 	void				End();
 	bool				OnData(FStreamReader& Reader);
@@ -4157,9 +4185,9 @@ public:
 };
 
 ////////////////////////////////////////////////////////////////////////////////
-FAnalysisEngine::FImpl::FImpl(TArray<IAnalyzer*>&& Analyzers, FMessageLog* InLog)
+FAnalysisEngine::FImpl::FImpl(TArray<IAnalyzer*>&& Analyzers, FMessageDelegate&& InMessage)
 : Bridge(Forward<TArray<IAnalyzer*>>(Analyzers))
-, Machine(Bridge, InLog)
+, Machine(Bridge, Forward<FMessageDelegate>(InMessage))
 {
 }
 
@@ -4188,8 +4216,8 @@ bool FAnalysisEngine::FImpl::OnData(FStreamReader& Reader)
 
 
 ////////////////////////////////////////////////////////////////////////////////
-FAnalysisEngine::FAnalysisEngine(TArray<IAnalyzer*>&& Analyzers, FMessageLog* InLog)
-: Impl(new FImpl(Forward<TArray<IAnalyzer*>>(Analyzers), InLog))
+FAnalysisEngine::FAnalysisEngine(TArray<IAnalyzer*>&& Analyzers, FMessageDelegate&& InMessage)
+: Impl(new FImpl(Forward<TArray<IAnalyzer*>>(Analyzers), Forward<FMessageDelegate>(InMessage)))
 {
 }
 
diff --git a/Engine/Source/Developer/TraceAnalysis/Private/Analysis/Engine.h b/Engine/Source/Developer/TraceAnalysis/Private/Analysis/Engine.h
--- a/Engine/Source/Developer/TraceAnalysis/Private/Analysis/Engine.h
+++ b/Engine/Source/Developer/TraceAnalysis/Private/Analysis/Engine.h
@@ -3,6 +3,7 @@
 #pragma once
 
 #include "Containers/ContainersFwd.h"
+#include "Trace/Analysis.h"
 
 class FMessageLog;
 
@@ -16,7 +17,7 @@ class FStreamReader;
 class FAnalysisEngine
 {
 public:
-						FAnalysisEngine(TArray<IAnalyzer*>&& Analyzers, FMessageLog* InLog);
+						FAnalysisEngine(TArray<IAnalyzer*>&& Analyzers, FMessageDelegate&& InMessage);
 						~FAnalysisEngine();
 	void				Begin();
 	void				End();
diff --git a/Engine/Source/Developer/TraceAnalysis/Private/Analysis/Processor.cpp b/Engine/Source/Developer/TraceAnalysis/Private/Analysis/Processor.cpp
--- a/Engine/Source/Developer/TraceAnalysis/Private/Analysis/Processor.cpp
+++ b/Engine/Source/Developer/TraceAnalysis/Private/Analysis/Processor.cpp
@@ -15,9 +15,8 @@ namespace UE {
 namespace Trace {
 
 ////////////////////////////////////////////////////////////////////////////////
-FAnalysisProcessor::FImpl::FImpl(IInDataStream& InDataStream, TArray<IAnalyzer*>&& InAnalyzers)
-: Log(FTraceAnalysisModule::GetMessageLogName())
-, AnalysisEngine(Forward<TArray<IAnalyzer*>>(InAnalyzers), &Log)
+FAnalysisProcessor::FImpl::FImpl(IInDataStream& InDataStream, TArray<IAnalyzer*>&& InAnalyzers, FMessageDelegate&& InMessage)
+: AnalysisEngine(Forward<TArray<IAnalyzer*>>(InAnalyzers), Forward<FMessageDelegate>(InMessage))
 , DataStream(InDataStream)
 , StopEvent(FPlatformProcess::GetSynchEventFromPool(true))
 , UnpausedEvent(FPlatformProcess::GetSynchEventFromPool(true))
@@ -103,21 +102,12 @@ void FAnalysisProcessor::FImpl::PauseAnalysis(bool bState)
 	}
 }
 
-////////////////////////////////////////////////////////////////////////////////
-FMessageLog* FAnalysisProcessor::FImpl::GetLog() 
-{
-	// MessageLog is not thread safe, so it's not safe to return it while
-	// messages could be pushed to the log.
-	return bComplete ? &Log : nullptr;
-}
-
 
 ////////////////////////////////////////////////////////////////////////////////
 bool FAnalysisProcessor::IsActive() const	{ return (Impl != nullptr) ? Impl->IsActive() : false; }
 void FAnalysisProcessor::Stop()				{ if (Impl != nullptr) { Impl->StopAnalysis(); } }
 void FAnalysisProcessor::Wait()				{ if (Impl != nullptr) { Impl->WaitOnAnalysis(); } }
 void FAnalysisProcessor::Pause(bool bState) { if (Impl != nullptr) { Impl->PauseAnalysis(bState); } }
-FMessageLog* FAnalysisProcessor::GetLog() { return Impl->GetLog(); }	
 
 ////////////////////////////////////////////////////////////////////////////////
 FAnalysisProcessor::FAnalysisProcessor(FAnalysisProcessor&& Rhs)
diff --git a/Engine/Source/Developer/TraceAnalysis/Private/Analysis/Processor.h b/Engine/Source/Developer/TraceAnalysis/Private/Analysis/Processor.h
--- a/Engine/Source/Developer/TraceAnalysis/Private/Analysis/Processor.h
+++ b/Engine/Source/Developer/TraceAnalysis/Private/Analysis/Processor.h
@@ -18,7 +18,7 @@ class FAnalysisProcessor::FImpl
 	: public FRunnable
 {
 public:
-						FImpl(IInDataStream& DataStream, TArray<IAnalyzer*>&& InAnalyzers);
+						FImpl(IInDataStream& DataStream, TArray<IAnalyzer*>&& InAnalyzers, FMessageDelegate&& InMessage);
 						~FImpl();
 	virtual uint32		Run() override;
 	bool				IsActive() const;
@@ -28,7 +28,6 @@ public:
 	FMessageLog*		GetLog();
 
 private:
-	FMessageLog			Log;
 	FAnalysisEngine		AnalysisEngine;
 	IInDataStream&		DataStream;
 	FEvent*				StopEvent;
diff --git a/Engine/Source/Developer/TraceAnalysis/Private/TraceAnalysisModule.cpp b/Engine/Source/Developer/TraceAnalysis/Private/TraceAnalysisModule.cpp
--- a/Engine/Source/Developer/TraceAnalysis/Private/TraceAnalysisModule.cpp
+++ b/Engine/Source/Developer/TraceAnalysis/Private/TraceAnalysisModule.cpp
@@ -2,15 +2,11 @@
 
 #include "TraceAnalysisModule.h"
 
-#include "MessageLogModule.h"
 #include "Modules/ModuleManager.h"
 
 ////////////////////////////////////////////////////////////////////////////////
 void FTraceAnalysisModule::StartupModule()
 {
-	FMessageLogModule& MessageLogModule = FModuleManager::LoadModuleChecked<FMessageLogModule>("MessageLog");
-	MessageLogModule.RegisterLogListing(GetMessageLogName(), NSLOCTEXT("Trace", "TraceAnalysis", "Trace Analysis"));
-	MessageLogModule.EnableMessageLogDisplay(true);
 }
 
 ////////////////////////////////////////////////////////////////////////////////
diff --git a/Engine/Source/Developer/TraceAnalysis/Public/Trace/Analysis.h b/Engine/Source/Developer/TraceAnalysis/Public/Trace/Analysis.h
--- a/Engine/Source/Developer/TraceAnalysis/Public/Trace/Analysis.h
+++ b/Engine/Source/Developer/TraceAnalysis/Public/Trace/Analysis.h
@@ -3,6 +3,7 @@
 #pragma once
 
 #include "Containers/Array.h"
+#include "Delegates/Delegate.h"
 
 class FMessageLog;
 
@@ -11,6 +12,15 @@ namespace Trace {
 
 class IAnalyzer;
 class IInDataStream;
+	
+enum class EAnalysisMessageSeverity
+{
+	Info,
+	Warning,
+	Error
+};
+
+DECLARE_DELEGATE_TwoParams(FMessageDelegate, EAnalysisMessageSeverity, FStringView);	
 
 /**
  * Represents the processing (e.g. analysis) of a trace stream. Instances are
@@ -35,11 +45,6 @@ public:
 	void Pause(bool bState);
 
 
-	/**
-	 * Gets the message log. Must be called after analysis has finished.
-	 */
-	FMessageLog* GetLog();
-
 						~FAnalysisProcessor();
 						FAnalysisProcessor() = default;
 						FAnalysisProcessor(FAnalysisProcessor&& Rhs);
@@ -66,6 +71,9 @@ public:
 	 * from the trace stream. */
 	void AddAnalyzer(IAnalyzer& Analyzer);
 
+	/** Adds a callback to recieve important messages. */
+	void SetMessageDelegate(FMessageDelegate Delegate);
+
 	/** Creates and starts analysis returning an FAnalysisProcessor instance which
 	 * represents the analysis and affords some control over it.
 	 * @param DataStream Input stream of trace log data to be analysed. */
@@ -73,6 +81,7 @@ public:
 
 private:
 	TArray<IAnalyzer*>	Analyzers;
+	FMessageDelegate	OnMessage;
 };
 
 } // namespace Trace
diff --git a/Engine/Source/Developer/TraceAnalysis/TraceAnalysis.Build.cs b/Engine/Source/Developer/TraceAnalysis/TraceAnalysis.Build.cs
--- a/Engine/Source/Developer/TraceAnalysis/TraceAnalysis.Build.cs
+++ b/Engine/Source/Developer/TraceAnalysis/TraceAnalysis.Build.cs
@@ -12,10 +12,7 @@ public class TraceAnalysis : ModuleRules
 		PrivateDependencyModuleNames.Add("Asio");
 		PrivateDependencyModuleNames.Add("Cbor");
 		PrivateDependencyModuleNames.Add("Core");
-		PrivateDependencyModuleNames.Add("MessageLog");
 		PrivateDependencyModuleNames.Add("Sockets");
 		PrivateDependencyModuleNames.Add("TraceLog");
-
-		PrivateIncludePathModuleNames.Add("MessageLog");
 	}
 }
diff --git a/Engine/Source/Developer/TraceInsights/Private/Insights/InsightsManager.cpp b/Engine/Source/Developer/TraceInsights/Private/Insights/InsightsManager.cpp
--- a/Engine/Source/Developer/TraceInsights/Private/Insights/InsightsManager.cpp
+++ b/Engine/Source/Developer/TraceInsights/Private/Insights/InsightsManager.cpp
@@ -143,6 +143,7 @@ TSharedPtr<FInsightsManager> FInsightsManager::CreateInstance(TSharedRef<TraceSe
 FInsightsManager::FInsightsManager(TSharedRef<TraceServices::IAnalysisService> InTraceAnalysisService,
 								   TSharedRef<TraceServices::IModuleService> InTraceModuleService)
 	: LogListingName(TEXT("UnrealInsights"))
+	, AnalysisLogListingName(TEXT("TraceAnalysis"))
 	, AnalysisService(InTraceAnalysisService)
 	, ModuleService(InTraceModuleService)
 	, CommandList(new FUICommandList())
@@ -167,6 +168,7 @@ void FInsightsManager::Initialize(IUnrealInsightsModule& InsightsModule)
 
 	FMessageLogModule& MessageLogModule = FModuleManager::LoadModuleChecked<FMessageLogModule>("MessageLog");
 	MessageLogModule.RegisterLogListing(GetLogListingName(), LOCTEXT("UnrealInsights", "Unreal Insights"));
+	MessageLogModule.RegisterLogListing(AnalysisLogListingName, LOCTEXT("TraceAnalysis", "Trace Analysis"));
 	MessageLogModule.EnableMessageLogDisplay(true);
 
 	// Register tick functions.
@@ -537,6 +539,8 @@ bool FInsightsManager::Tick(float DeltaTime)
 
 	UpdateSessionDuration();
 
+	PollAnalysisInfo();
+
 #if !WITH_EDITOR
 	if (!bIsSessionInfoSet && Session.IsValid())
 	{
@@ -576,8 +580,10 @@ void FInsightsManager::UpdateSessionDuration()
 			TraceServices::FAnalysisSessionReadScope SessionReadScope(*Session.Get());
 			bLocalIsAnalysisComplete = Session->IsAnalysisComplete();
 			LocalSessionDuration = Session->GetDurationSeconds();
+
 		}
 
+
 		if (LocalSessionDuration != SessionDuration)
 		{
 			SessionDuration = LocalSessionDuration;
@@ -685,8 +691,39 @@ void FInsightsManager::ResetSession(bool bNotify)
 
 ////////////////////////////////////////////////////////////////////////////////////////////////////
 
+void FInsightsManager::PollAnalysisInfo()
+{
+	if (Session.IsValid())
+	{
+		if (Session->GetNumPendingMessages())
+		{
+			TSharedPtr<TraceServices::IAnalysisSession> EditableSession = ConstCastSharedPtr<TraceServices::IAnalysisSession>(Session);
+			TraceServices::FAnalysisSessionEditScope SessionEditScope(*EditableSession.Get());
+			const auto Messages = EditableSession->DrainPendingMessages();
+			
+			FMessageLog ReportMessageLog(AnalysisLogListingName);
+			for (const auto& Message : Messages)
+			{
+				TSharedRef<FTokenizedMessage> LogMessage = FTokenizedMessage::Create(Message.Severity, FText::FromString(Message.Message));
+				ReportMessageLog.AddMessage(LogMessage);
+				if (Message.Severity == EMessageSeverity::Error)
+				{
+					ReportMessageLog.Notify();
+				}
+			}
+		}
+	}
+}
+
+////////////////////////////////////////////////////////////////////////////////////////////////////
+
 void FInsightsManager::OnSessionChanged()
 {
+	if (Session.IsValid())
+	{
+		FMessageLog(AnalysisLogListingName).NewPage(FText::FromString(Session->GetName()));
+	}
+	
 	SessionChangedEvent.Broadcast();
 }
 
@@ -1133,26 +1170,6 @@ void FInsightsManager::ScheduleCommand(const FString& InCmd)
 
 void FInsightsManager::OnSessionAnalysisCompleted()
 {
-	// Report any errors or warning that occurred
-	FMessageLog* Log = Session->GetLog();
-	if (Log)
-	{
-		if (const int32 NumErrors = Log->NumMessages(EMessageSeverity::Error); NumErrors > 0)
-		{
-			Log->Notify(FText::Format(NSLOCTEXT("TraceAnalysis", "FinishedWithErrors", "Analysis finished with {0} error(s)"),
-				FText::AsNumber(NumErrors)));
-		}
-		else if (const int32 NumWarnings = Log->NumMessages(EMessageSeverity::Warning); NumWarnings > 0)
-		{
-			Log->Notify(FText::Format(NSLOCTEXT("TraceAnalysis", "FinishedWithWarnings", "Analysis finished with {0} warning(s)"),
-				FText::AsNumber(NumWarnings)));
-		}
-		else
-		{
-			Log->Notify(NSLOCTEXT("TraceAnalysis", "Finished", "Analysis finished."));
-		}
-	}
-
 	if (!SessionAnalysisCompletedCmd.IsEmpty())
 	{
 		FOutputDevice& Ar = *GLog;
diff --git a/Engine/Source/Developer/TraceInsights/Private/Insights/InsightsManager.h b/Engine/Source/Developer/TraceInsights/Private/Insights/InsightsManager.h
--- a/Engine/Source/Developer/TraceInsights/Private/Insights/InsightsManager.h
+++ b/Engine/Source/Developer/TraceInsights/Private/Insights/InsightsManager.h
@@ -390,6 +390,9 @@ private:
 	/** Resets (closes) current session instance. */
 	void ResetSession(bool bNotify = true);
 
+	/** Extract messages from the session */
+	void PollAnalysisInfo();
+
 	void OnSessionChanged();
 	void OnSessionAnalysisCompleted();
 
@@ -411,6 +414,9 @@ private:
 	/** The name of the Unreal Insights log listing. */
 	FName LogListingName;
 
+	/** Name used for Analysis log in Message Log. */
+	FName AnalysisLogListingName;
+
 	/** The delegate to be invoked when this manager ticks. */
 	FTickerDelegate OnTick;
 
diff --git a/Engine/Source/Developer/TraceServices/Private/AnalysisService.cpp b/Engine/Source/Developer/TraceServices/Private/AnalysisService.cpp
--- a/Engine/Source/Developer/TraceServices/Private/AnalysisService.cpp
+++ b/Engine/Source/Developer/TraceServices/Private/AnalysisService.cpp
@@ -113,6 +113,7 @@ void FAnalysisSession::Start()
 	{
 		Context.AddAnalyzer(*Analyzer);
 	}
+	Context.SetMessageDelegate(UE::Trace::FMessageDelegate::CreateRaw(this, &FAnalysisSession::OnAnalysisMessage));
 	Processor = Context.Process(*DataStream);
 }
 
@@ -167,9 +168,16 @@ void FAnalysisSession::AddMetadata(FName InName, FString InValue)
 	Value.StringValue = InValue;
 }
 
-FMessageLog* FAnalysisSession::GetLog() const
+uint32 FAnalysisSession::GetNumPendingMessages() const
 {
-	return Processor.GetLog();
+	return PendingMessagesCount.load();
+}
+
+TArray<FAnalysisMessage> FAnalysisSession::DrainPendingMessages() 
+{
+	Lock.WriteAccessCheck();
+	PendingMessagesCount.store(0);
+	return MoveTemp(PendingMessages);
 }
 
 void FAnalysisSession::AddAnalyzer(UE::Trace::IAnalyzer* Analyzer)
@@ -208,6 +216,22 @@ IEditableProvider* FAnalysisSession::EditProviderPrivate(const FName& InName)
 	}
 }
 
+void FAnalysisSession::OnAnalysisMessage(UE::Trace::EAnalysisMessageSeverity InSeverity, FStringView InMessage)
+{
+	EMessageSeverity::Type Severity = EMessageSeverity::Type::Info;
+	switch(InSeverity)
+	{
+	case UE::Trace::EAnalysisMessageSeverity::Error: Severity = EMessageSeverity::Type::Error; break;
+	case UE::Trace::EAnalysisMessageSeverity::Warning: Severity = EMessageSeverity::Type::Warning; break;
+	case UE::Trace::EAnalysisMessageSeverity::Info: Severity = EMessageSeverity::Type::Info; break;
+	}
+	
+	Lock.BeginEdit();
+	PendingMessages.Push(FAnalysisMessage { Severity, FString(InMessage)});
+	PendingMessagesCount.fetch_add(1);
+	Lock.EndEdit();
+}
+
 FAnalysisService::FAnalysisService(FModuleService& InModuleService)
 	: ModuleService(InModuleService)
 {
diff --git a/Engine/Source/Developer/TraceServices/Private/AnalysisServicePrivate.h b/Engine/Source/Developer/TraceServices/Private/AnalysisServicePrivate.h
--- a/Engine/Source/Developer/TraceServices/Private/AnalysisServicePrivate.h
+++ b/Engine/Source/Developer/TraceServices/Private/AnalysisServicePrivate.h
@@ -55,7 +55,9 @@ public:
 	virtual void AddMetadata(FName InName, double InValue) override;
 	virtual void AddMetadata(FName InName, FString InValue) override;
 
-	virtual FMessageLog* GetLog() const override;
+	virtual FMessageLog* GetLog() const override {return nullptr;};
+	virtual uint32 GetNumPendingMessages() const override;
+	virtual TArray<FAnalysisMessage> DrainPendingMessages() override;
 
 	virtual ILinearAllocator& GetLinearAllocator() override { return Allocator; }
 	virtual const TCHAR* StoreString(const TCHAR* String) override { return StringStore.Store(String); }
@@ -81,6 +83,8 @@ private:
 	virtual const IProvider* ReadProviderPrivate(const FName& Name) const override;
 	virtual IEditableProvider* EditProviderPrivate(const FName& Name) override;
 
+	void OnAnalysisMessage(UE::Trace::EAnalysisMessageSeverity Severity, FStringView Message);
+
 private:
 	mutable FAnalysisSessionLock Lock;
 	FString Name;
@@ -92,6 +96,8 @@ private:
 	FAnalysisCache Cache;
 	TArray<UE::Trace::IAnalyzer*> Analyzers;
 	TMap<FName, TTuple<TSharedPtr<IProvider>, TSharedPtr<IEditableProvider>>> Providers;
+	TArray<FAnalysisMessage> PendingMessages;
+	mutable std::atomic<uint32> PendingMessagesCount;
 	mutable TUniquePtr<UE::Trace::IInDataStream> DataStream;
 	mutable UE::Trace::FAnalysisProcessor Processor;
 };
diff --git a/Engine/Source/Developer/TraceServices/Private/Model/Regions.cpp b/Engine/Source/Developer/TraceServices/Private/Model/Regions.cpp
--- a/Engine/Source/Developer/TraceServices/Private/Model/Regions.cpp
+++ b/Engine/Source/Developer/TraceServices/Private/Model/Regions.cpp
@@ -7,7 +7,6 @@
 #include "Common/FormatArgs.h"
 #include "Common/Utils.h"
 #include "Internationalization/Internationalization.h"
-#include "Logging/MessageLog.h"
 
 #define LOCTEXT_NAMESPACE "RegionProvider"
 
@@ -53,14 +52,7 @@ void FRegionProvider::AppendRegionBegin(const TCHAR* Name, double Time)
 		++NumWarnings;
 		if (NumWarnings <= MaxWarningMessages)
 		{
-			if (FMessageLog* Log = Session.GetLog())
-			{
-				Log->Warning(FText::Format(LOCTEXT("DoubleRegionBegin", "A region begin event ({0}) was encountered while a region with same name is already open."), FText::FromString(FString(Name))));
-			}
-			else
-			{
-				UE_LOG(LogTraceServices, Warning, TEXT("[Regions] A region begin event (%s) was encountered while a region with same name is already open."), Name)
-			}
+			UE_LOG(LogTraceServices, Warning, TEXT("[Regions] A region begin event (%s) was encountered while a region with same name is already open."), Name)
 		}
 	}
 	else
@@ -105,14 +97,7 @@ void FRegionProvider::AppendRegionEnd(const TCHAR* Name, double Time)
 		++NumWarnings;
 		if (NumWarnings <= MaxWarningMessages)
 		{
-			if (FMessageLog* Log = Session.GetLog())
-			{
-				Log->Warning(FText::Format(LOCTEXT("StrayRegionEnd", "A region end event ({0}) was encountered without having seen a matching region start event first."), FText::FromString(FString(Name))));
-			}
-			else
-			{
-				UE_LOG(LogTraceServices, Warning, TEXT("[Regions] A region end event (%s) was encountered without having seen a matching region start event first."), Name)
-			}
+			UE_LOG(LogTraceServices, Warning, TEXT("[Regions] A region end event (%s) was encountered without having seen a matching region start event first."), Name)
 		}
 	}
 
@@ -134,14 +119,7 @@ void FRegionProvider::OnAnalysisSessionEnded()
 		++NumWarnings;
 		if (NumWarnings <= MaxWarningMessages)
 		{
-			if (FMessageLog* Log = Session.GetLog())
-			{
-				Log->Warning(FText::Format(LOCTEXT("StrayRegionBegin", "A region begin event ({0}) was never closed."), FText::FromString(FString(Region->Text))));
-			}
-			else
-			{
-				UE_LOG(LogTraceServices, Warning, TEXT("[Regions] A region begin event (%s) was never closed."), Region->Text)
-			}
+			UE_LOG(LogTraceServices, Warning, TEXT("[Regions] A region begin event (%s) was never closed."), Region->Text)
 		}
 	}
 
diff --git a/Engine/Source/Developer/TraceServices/Public/TraceServices/Model/AnalysisSession.h b/Engine/Source/Developer/TraceServices/Public/TraceServices/Model/AnalysisSession.h
--- a/Engine/Source/Developer/TraceServices/Public/TraceServices/Model/AnalysisSession.h
+++ b/Engine/Source/Developer/TraceServices/Public/TraceServices/Model/AnalysisSession.h
@@ -4,6 +4,7 @@
 
 #include "CoreTypes.h"
 #include "Containers/StringFwd.h"
+#include "Logging/TokenizedMessage.h"
 #include "Misc/AssertionMacros.h"
 #include "Templates/SharedPointer.h"
 
@@ -63,9 +64,17 @@ struct FTraceSessionMetadata
 	FString StringValue;
 };
 
+struct FAnalysisMessage
+{
+	EMessageSeverity::Type Severity;
+	FString Message;	
+};
+
 class IAnalysisSession
 {
 public:
+
+	
 	virtual ~IAnalysisSession() = default;
 	virtual void Stop(bool bAndWait=false) const = 0;
 	virtual void Wait() const = 0;
@@ -88,7 +97,12 @@ public:
 	virtual void AddMetadata(FName InName, double InValue) = 0;
 	virtual void AddMetadata(FName InName, FString InValue) = 0;
 
+	UE_DEPRECATED(5.3, "No longer used.")
 	virtual FMessageLog* GetLog() const = 0;
+	/** Get the number of pending messages. Can be called without lock */
+	virtual uint32 GetNumPendingMessages() const = 0;
+	/** Moves pending messages to the return array. */
+	virtual TArray<FAnalysisMessage> DrainPendingMessages() = 0;
 
 	virtual void BeginRead() const = 0;
 	virtual void EndRead() const = 0;
