[Insights] Warkaround for corrupted traces with missing sync events. Analysis will not stop anymore when buffering more than 512 MiB of trace data. Instead, it will emit an error and then it will start to skip the missing serial numbers.

diff --git a/Engine/Source/Developer/TraceAnalysis/Private/Analysis/Engine.cpp b/Engine/Source/Developer/TraceAnalysis/Private/Analysis/Engine.cpp
--- a/Engine/Source/Developer/TraceAnalysis/Private/Analysis/Engine.cpp
+++ b/Engine/Source/Developer/TraceAnalysis/Private/Analysis/Engine.cpp
@@ -2916,6 +2916,10 @@ protected:
 	{
 		bool operator () (const FEventDescStream& Lhs, const FEventDescStream& Rhs) const
 		{
+			// Provided that less than approximately "SerialRange * BytesPerSerial"
+			// is buffered there should never be more that "SerialRange / 2" serial
+			// numbers. Thus if the distance between any two serial numbers is larger
+			// than half the serial space, they have wrapped.
 			uint32 Ld = Lhs.EventDescs->Serial - Origin;
 			uint32 Rd = Rhs.EventDescs->Serial - Origin;
 			return Ld < Rd;
@@ -2946,8 +2950,9 @@ protected:
 	int32					ParseEvent(FStreamReader& Reader, FEventDesc& OutEventDesc, const FMachineContext& Context);
 	virtual void			SetSizeIfKnownEvent(uint32 Uid, uint32& InOutEventSize);
 	virtual bool			DispatchKnownEvent(const FMachineContext& Context, uint32 Uid, const FEventDesc* Cursor);
-	int32					DispatchEvents(const FMachineContext& Context, const FEventDesc* EventDesc, uint32 Count);
+	int32					DispatchNormalEvents(const FMachineContext& Context, TArray<FEventDescStream>& EventDescHeap);
 	int32					DispatchEvents(const FMachineContext& Context, TArray<FEventDescStream>& EventDescHeap);
+	int32					DispatchEvents(const FMachineContext& Context, const FEventDesc* EventDesc, uint32 Count);
 	void					DetectSerialGaps(TArray<FEventDescStream>& EventDescHeap);
 	template <typename Callback>
 	void					ForEachSerialGap(const TArray<FEventDescStream>& EventDescHeap, Callback&& InCallback);
@@ -2959,6 +2964,7 @@ protected:
 	uint32					NextSerial = ~0u;
 	uint32					SyncCount;
 	uint32					EventVersion = 4; //Protocol version 5 uses the event version from protocol 4
+	bool					bSkipSerialError = false;
 };
 
 ////////////////////////////////////////////////////////////////////////////////
@@ -3180,6 +3186,8 @@ FProtocol5Stage::EStatus FProtocol5Stage::OnDataNormal(const FMachineContext& Co
 	TArray<FEventDescStream> EventDescHeap;
 	EventDescHeap.Reserve(Transport.GetThreadCount());
 
+	bool bSkipSerial = false;
+
 	for (uint32 i = ETransportTid::Bias, n = Transport.GetThreadCount(); i < n; ++i)
 	{
 		uint32 NumEventDescs = EventDescs.Num();
@@ -3194,12 +3202,17 @@ FProtocol5Stage::EStatus FProtocol5Stage::OnDataNormal(const FMachineContext& Co
 		constexpr int32 MaxAccumulatedBytes = 512 * 1024 * 1024;
 		if (ThreadReader->GetRemaining() > MaxAccumulatedBytes)
 		{
-			Context.EmitMessagef(
-				EAnalysisMessageSeverity::Error,
-				TEXT("Analysis accumulated too much data (%d MiB) before being able to continue analysing the stream."),
-				ThreadReader->GetRemaining() / (1024*1024)
-			);
-			return EStatus::Error;
+			if (!bSkipSerialError)
+			{
+				bSkipSerialError = true;
+				Context.EmitMessagef(
+					EAnalysisMessageSeverity::Error,
+					TEXT("Trace analysis accumulated too much data (%d MiB on thread %u) and will start to skip the missing serial sync events!"),
+					ThreadReader->GetRemaining() / (1024 * 1024),
+					Transport.GetThreadId(i)
+				);
+			}
+			bSkipSerial = true;
 		}
 
 		if (ParseEvents(*ThreadReader, EventDescs, Context) < 0)
@@ -3231,8 +3244,68 @@ FProtocol5Stage::EStatus FProtocol5Stage::OnDataNormal(const FMachineContext& Co
 		Stream.EventDescs = EventDescs.GetData() + Stream.ContainerIndex;
 	}
 
-	// Process leading unsynchronised events so that each stream starts with a
-	// sychronised event.
+	const bool bSync = (SyncCount != Transport.GetSyncCount());
+
+	if (bSync || !bSkipSerial)
+	{
+		int32 NumDispatchedEvents = DispatchNormalEvents(Context, EventDescHeap);
+		if (NumDispatchedEvents < 0)
+		{
+			return EStatus::Error;
+		}
+	}
+	else // bSkipSerial && !bSync
+	{
+		// Dispatch at least half of the parsed events.
+		int32 TotalNumDispatchedEvents = 0;
+		do
+		{
+			NextSerial = ~0u; // skip missing serial sync events
+
+			int32 NumDispatchedEvents = DispatchNormalEvents(Context, EventDescHeap);
+			if (NumDispatchedEvents < 0)
+			{
+				return EStatus::Error;
+			}
+
+			TotalNumDispatchedEvents += NumDispatchedEvents;
+			if (TotalNumDispatchedEvents > EventDescs.Num() / 2)
+			{
+				bSkipSerial = false;
+			}
+		}
+		while (bSkipSerial);
+	}
+
+	// If there are any streams left in the heap then we are unable to proceed
+	// until more data is received. We'll rewind the streams until more data is
+	// available. It is not an efficient way to do things, but it is simple way.
+	for (FEventDescStream& Stream : EventDescHeap)
+	{
+		const FEventDesc& EventDesc = Stream.EventDescs[0];
+		uint32 HeaderSize = 1 + EventDesc.bTwoByteUid + (ESerial::Bits / 8);
+
+		FStreamReader* Reader = Transport.GetThreadStream(Stream.TransportIndex);
+		Reader->Backtrack(EventDesc.Data - HeaderSize);
+	}
+
+	if (bSync && SyncCount == Transport.GetSyncCount())
+	{
+		return EStatus::Sync;
+	}
+	if (bNotEnoughData)
+	{
+		return EStatus::NotEnoughData;
+	}
+	return EStatus::EndOfStream;
+}
+
+////////////////////////////////////////////////////////////////////////////////
+int32 FProtocol5Stage::DispatchNormalEvents(const FMachineContext& Context, TArray<FEventDescStream>& EventDescHeap)
+{
+	int32 NumDispatchedUnsyncEvents = 0;
+
+	// Process leading unsynchronised events so that each stream starts with a sychronised event.
 	for (FEventDescStream& Stream : EventDescHeap)
 	{
 		// Extract a run of consecutive unsynchronised events
@@ -3244,11 +3317,13 @@ FProtocol5Stage::EStatus FProtocol5Stage::OnDataNormal(const FMachineContext& Co
 		int32 DescNum = int32(UPTRINT(EndDesc - StartDesc));
 		if (DescNum > 0)
 		{
+			NumDispatchedUnsyncEvents += DescNum;
+
 			Context.Bridge.SetActiveThread(Stream.ThreadId);
 
 			if (DispatchEvents(Context, StartDesc, DescNum) < 0)
 			{
-				return EStatus::Error;
+				return -1;
 			}
 
 			Stream.EventDescs = EndDesc;
@@ -3264,14 +3339,9 @@ FProtocol5Stage::EStatus FProtocol5Stage::OnDataNormal(const FMachineContext& Co
 	// Early out if there isn't any events available.
 	if (UNLIKELY(EventDescHeap.IsEmpty()))
 	{
-		return bNotEnoughData ? EStatus::NotEnoughData : EStatus::EndOfStream;
+		return NumDispatchedUnsyncEvents;
 	}
 
-	// Provided that less than approximately "SerialRange * BytesPerSerial"
-	// is buffered there should never be more that "SerialRange / 2" serial
-	// numbers. Thus if the distance between any two serial numbers is larger
-	// than half the serial space, they have wrapped.
-
 	// A min-heap is used to peel off groups of events by lowest serial
 	EventDescHeap.Heapify(FSerialDistancePredicate{NextSerial});
 
@@ -3281,20 +3351,15 @@ FProtocol5Stage::EStatus FProtocol5Stage::OnDataNormal(const FMachineContext& Co
 		NextSerial = EventDescHeap.HeapTop().EventDescs[0].Serial;
 	}
 
-	const bool bSync = (SyncCount != Transport.GetSyncCount());
 	DetectSerialGaps(EventDescHeap);
 
-	if (DispatchEvents(Context, EventDescHeap) < 0)
+	int32 NumDispatchedEvents = DispatchEvents(Context, EventDescHeap);
+	if (NumDispatchedEvents < 0)
 	{
-		return EStatus::Error;
-	}
-
-	if (bSync)
-	{
-		return EStatus::Sync;
+		return -1;
 	}
 
-	return bNotEnoughData ? EStatus::NotEnoughData : EStatus::EndOfStream;
+	return NumDispatchedUnsyncEvents + NumDispatchedEvents;
 }
 
 ////////////////////////////////////////////////////////////////////////////////
@@ -3314,6 +3379,8 @@ int32 FProtocol5Stage::DispatchEvents(
 		EventDescHeap.HeapPopDiscard(FSerialDistancePredicate{NextSerial}, false);
 	};
 
+	int32 NumDispatchedEvents = 0;
+
 	do
 	{
 		const FEventDescStream& Stream = EventDescHeap.HeapTop();
@@ -3357,6 +3424,7 @@ int32 FProtocol5Stage::DispatchEvents(
 		Context.Bridge.SetActiveThread(Stream.ThreadId);
 		int32 DescNum = int32(UPTRINT(EndDesc - StartDesc));
 		check(DescNum > 0);
+		NumDispatchedEvents += DescNum;
 		if (DispatchEvents(Context, StartDesc, DescNum) < 0)
 		{
 			return -1;
@@ -3366,19 +3434,7 @@ int32 FProtocol5Stage::DispatchEvents(
 	}
 	while (!EventDescHeap.IsEmpty());
 
-	// If there are any streams left in the heap then we are unable to proceed
-	// until more data is received. We'll rewind the streams until more data is
-	// available. It is an efficient way to do things, but it is simple way.
-	for (FEventDescStream& Stream : EventDescHeap)
-	{
-		const FEventDesc& EventDesc = Stream.EventDescs[0];
-		uint32 HeaderSize = 1 + EventDesc.bTwoByteUid + (ESerial::Bits / 8);
-
-		FStreamReader* Reader = Transport.GetThreadStream(Stream.TransportIndex);
-		Reader->Backtrack(EventDesc.Data - HeaderSize);
-	}
-
-	return 0;
+	return NumDispatchedEvents;
 }
 
 ////////////////////////////////////////////////////////////////////////////////
