Display store server version in Unreal Insights

Uses the version enpoint in Unreal Trace Server to display the server version in the connection tooltip.

diff --git a/Engine/Source/Developer/TraceAnalysis/Private/Store/StoreClient.cpp b/Engine/Source/Developer/TraceAnalysis/Private/Store/StoreClient.cpp
--- a/Engine/Source/Developer/TraceAnalysis/Private/Store/StoreClient.cpp
+++ b/Engine/Source/Developer/TraceAnalysis/Private/Store/StoreClient.cpp
@@ -107,6 +107,7 @@ public:
 	const FResponse&		GetResponse() const;
 	bool					Connect(const TCHAR* Host, uint16 Port);
 	bool					GetStatus();
+	bool					GetVersion();
 	bool					GetTraceCount();
 	bool					GetTraceInfo(uint32 Index);
 	bool					GetTraceInfoById(uint32 Id);
@@ -280,6 +281,14 @@ bool FStoreCborClient::GetStatus()
 	return Communicate(Payload);
 }
 
+////////////////////////////////////////////////////////////////////////////////
+bool FStoreCborClient::GetVersion()
+{
+	TPayloadBuilder<32> Builder("v1/version");
+	FPayload Payload = Builder.Done();
+	return Communicate(Payload);
+}
+	
 ////////////////////////////////////////////////////////////////////////////////
 bool FStoreCborClient::GetTraceCount()
 {
@@ -455,6 +464,26 @@ void FStoreClient::FStatus::GetWatchDirectories(TArray<FString>& OutDirs) const
 	Response->GetStringArray("watch_dirs", OutDirs);
 }
 
+////////////////////////////////////////////////////////////////////////////////
+uint32 FStoreClient::FVersion::GetMajorVersion() const
+{
+	const auto* Response = (const FResponse*)this;
+	return Response->GetInteger("major", 0);
+}
+
+////////////////////////////////////////////////////////////////////////////////
+uint32 FStoreClient::FVersion::GetMinorVersion() const
+{
+	const auto* Response = (const FResponse*)this;
+	return Response->GetInteger("minor", 0);
+}
+
+////////////////////////////////////////////////////////////////////////////////
+FUtf8StringView FStoreClient::FVersion::GetConfiguration() const
+{
+	const auto* Response = (const FResponse*)this;
+	return Response->GetString("configuration", "unknown");
+}
 
 ////////////////////////////////////////////////////////////////////////////////
 uint32 FStoreClient::FTraceInfo::GetId() const
@@ -586,6 +615,20 @@ const FStoreClient::FStatus* FStoreClient::GetStatus()
 	return (FStatus*)(&Response);
 }
 
+////////////////////////////////////////////////////////////////////////////////
+const FStoreClient::FVersion* FStoreClient::GetVersion()
+{
+	
+	auto* Self = (FStoreCborClient*)this;
+	if (!Self->GetVersion())
+	{
+		return nullptr;
+	}
+
+	const FResponse& Response = Self->GetResponse();
+	return (FVersion*)(&Response);
+}
+
 ////////////////////////////////////////////////////////////////////////////////
 uint32 FStoreClient::GetTraceCount()
 {
diff --git a/Engine/Source/Developer/TraceAnalysis/Public/Trace/StoreClient.h b/Engine/Source/Developer/TraceAnalysis/Public/Trace/StoreClient.h
--- a/Engine/Source/Developer/TraceAnalysis/Public/Trace/StoreClient.h
+++ b/Engine/Source/Developer/TraceAnalysis/Public/Trace/StoreClient.h
@@ -26,6 +26,13 @@ public:
 		void			GetWatchDirectories(TArray<FString>& OutDirs) const;
 	};
 
+	struct TRACEANALYSIS_API FVersion
+	{
+		uint32			GetMajorVersion() const;
+		uint32			GetMinorVersion() const;
+		FUtf8StringView GetConfiguration() const;
+	};
+
 	struct TRACEANALYSIS_API FTraceInfo
 	{
 		FUtf8StringView	GetName() const;
@@ -51,6 +58,7 @@ public:
 	uint32				GetStoreAddress() const;
 	uint32				GetStorePort() const;
 	const FStatus*		GetStatus();
+	const FVersion*		GetVersion();
 	uint32				GetTraceCount();
 	const FTraceInfo*	GetTraceInfo(uint32 Index);
 	const FTraceInfo*	GetTraceInfoById(uint32 Id);
diff --git a/Engine/Source/Developer/TraceInsights/Private/Insights/StoreService/StoreBrowser.cpp b/Engine/Source/Developer/TraceInsights/Private/Insights/StoreService/StoreBrowser.cpp
--- a/Engine/Source/Developer/TraceInsights/Private/Insights/StoreService/StoreBrowser.cpp
+++ b/Engine/Source/Developer/TraceInsights/Private/Insights/StoreService/StoreBrowser.cpp
@@ -156,6 +156,27 @@ void FStoreBrowser::UpdateTraces()
 		{
 			FScopeLock Lock(&TracesCriticalSection);
 			Host = FInsightsManager::Get()->GetLastStoreHost();
+
+			// Get the server version
+			{
+				FScopeLock StoreClientLock(&GetStoreClientCriticalSection());
+				const auto VersionResult = StoreClient->GetVersion();
+
+				if (VersionResult)
+				{
+					TStringBuilder<128> VersionString;
+					VersionString << VersionResult->GetMajorVersion() << TEXT(".") << VersionResult->GetMinorVersion();
+					if (!VersionResult->GetConfiguration().Equals("Release"))
+					{
+						VersionString << TEXT(" (") << VersionResult->GetConfiguration() << TEXT(")");
+					}
+					Version = VersionString.ToString();
+				}
+				else
+				{
+					Version = TEXT("Unknown");
+				}
+			}
 		}
 	}
 
diff --git a/Engine/Source/Developer/TraceInsights/Private/Insights/StoreService/StoreBrowser.h b/Engine/Source/Developer/TraceInsights/Private/Insights/StoreService/StoreBrowser.h
--- a/Engine/Source/Developer/TraceInsights/Private/Insights/StoreService/StoreBrowser.h
+++ b/Engine/Source/Developer/TraceInsights/Private/Insights/StoreService/StoreBrowser.h
@@ -106,6 +106,7 @@ public:
 	const TArray<FString>& GetLockedWatchDirectories() const { check(bTracesLocked); return WatchDirectories; }
 	const FString& GetLockedStoreDirectory() const { check(bTracesLocked); return StoreDirectory; }
 	const FString& GetLockedHost() const { check(bTracesLocked); return Host; }
+	const FString& GetVersion() const { return Version; }
 	void Unlock() { check(bTracesLocked); bTracesLocked = false; TracesCriticalSection.Unlock(); }
 	void Refresh();
 
@@ -137,6 +138,7 @@ private:
 	TArray<FString> WatchDirectories;
 	FString StoreDirectory;
 	FString Host;
+	FString Version;
 };
 
 ////////////////////////////////////////////////////////////////////////////////////////////////////
diff --git a/Engine/Source/Developer/TraceInsights/Private/Insights/Widgets/SStartPageWindow.cpp b/Engine/Source/Developer/TraceInsights/Private/Insights/Widgets/SStartPageWindow.cpp
--- a/Engine/Source/Developer/TraceInsights/Private/Insights/Widgets/SStartPageWindow.cpp
+++ b/Engine/Source/Developer/TraceInsights/Private/Insights/Widgets/SStartPageWindow.cpp
@@ -1897,16 +1897,17 @@ FText STraceStoreWindow::GetConnectionStatusTooltip() const
 {
 	using Insights::FStoreBrowser;
 
-	static FText Connected    = LOCTEXT("Connected",    "Connected to the trace server.");
+	static FText Connected    = LOCTEXT("Connected",    "Connected to the trace server.\nServer version: {0}");
 	static FText NotConnected = LOCTEXT("NoConnection", "Unable to connect to trace server.");
 	static FText Connecting   = LOCTEXT("Connecting",   "Trying to connect to trace server.");
 	static FText Disconnected = LOCTEXT("Disconnected", "Connection to trace server has been lost. Attempting to reconnect in {0} seconds.");
 
 	const FStoreBrowser::EConnectionStatus Status = StoreBrowser->GetConnectionStatus();
+	const FString& Version = StoreBrowser->GetVersion();
 	switch (Status)
 	{
 		case FStoreBrowser::EConnectionStatus::Connected:
-			return Connected;
+			return FText::Format(Connected, FText::FromString(Version));
 
 		case FStoreBrowser::EConnectionStatus::NoConnection:
 			return NotConnected;
