Enabled UnsafeTypeCast warnings as errors in Trace/Insights related modules (UnrealInsights, TraceAnalysis, TraceServices, TraceDataFiltering, TraceUtilities, EditorTraceUtilities).
Fixed conversion warnings in TraceAnalysis and in TraceServices modules.

diff --git a/Engine/Plugins/Developer/TraceDataFiltering/Source/TraceDataFiltering/TraceDataFiltering.build.cs b/Engine/Plugins/Developer/TraceDataFiltering/Source/TraceDataFiltering/TraceDataFiltering.build.cs
--- a/Engine/Plugins/Developer/TraceDataFiltering/Source/TraceDataFiltering/TraceDataFiltering.build.cs
+++ b/Engine/Plugins/Developer/TraceDataFiltering/Source/TraceDataFiltering/TraceDataFiltering.build.cs
@@ -6,24 +6,26 @@ public class TraceDataFiltering : ModuleRules
 {
 	public TraceDataFiltering(ReadOnlyTargetRules Target) : base(Target)
 	{
+		UnsafeTypeCastWarningLevel = WarningLevel.Error;
+
 		PrivateDependencyModuleNames.AddRange(
-                new string[] {
-                    "Core",
+				new string[] {
+					"Core",
 					"CoreUObject",
-                    "SlateCore",
-                    "Slate",
-                    "InputCore",
-                    "TraceLog",
+					"InputCore",
+					"Slate",
+					"SlateCore",
+					"Sockets",
 					"TraceAnalysis",
+					"TraceInsights",
+					"TraceLog",
 					"TraceServices",
-                    "TraceInsights",
-					"Sockets"
-                }
-            );
+				}
+			);
 
 		if (Target.bBuildEditor)
-        {
-            PrivateDependencyModuleNames.Add("SharedSettingsWidgets");
-        }
-    }
+		{
+			PrivateDependencyModuleNames.Add("SharedSettingsWidgets");
+		}
+	}
 }
diff --git a/Engine/Plugins/TraceUtilities/Source/EditorTraceUtilities/EditorTraceUtilities.Build.cs b/Engine/Plugins/TraceUtilities/Source/EditorTraceUtilities/EditorTraceUtilities.Build.cs
--- a/Engine/Plugins/TraceUtilities/Source/EditorTraceUtilities/EditorTraceUtilities.Build.cs
+++ b/Engine/Plugins/TraceUtilities/Source/EditorTraceUtilities/EditorTraceUtilities.Build.cs
@@ -4,36 +4,35 @@ using UnrealBuildTool;
 
 public class EditorTraceUtilities : ModuleRules
 {
-    public EditorTraceUtilities(ReadOnlyTargetRules Target) : base(Target)
-    {
-        PCHUsage = ModuleRules.PCHUsageMode.UseExplicitOrSharedPCHs;
+	public EditorTraceUtilities(ReadOnlyTargetRules Target) : base(Target)
+	{
+		UnsafeTypeCastWarningLevel = WarningLevel.Error;
 
-        PublicDependencyModuleNames.AddRange(
-            new string[]
-            {
-                "Core",
-            }
-        );
+		PCHUsage = ModuleRules.PCHUsageMode.UseExplicitOrSharedPCHs;
+
+		PublicDependencyModuleNames.AddRange(
+			new string[] {
+				"Core",
+			}
+		);
 
 		PrivateIncludePathModuleNames.AddRange(
-			new string[]
-			{
+			new string[] {
 				"MessageLog",
 			}
 		);
 
 		PrivateDependencyModuleNames.AddRange(
-            new string[]
-            {
-                "Engine",
+			new string[] {
+				"Engine",
 				"InputCore",
 				"Slate",
-                "SlateCore",
-                "ToolMenus",
-                "UATHelper",
-                "TraceAnalysis",
+				"SlateCore",
+				"ToolMenus",
+				"TraceAnalysis",
 				"TraceLog",
+				"UATHelper",
 			}
-        );
-    }
+		);
+	}
 }
diff --git a/Engine/Plugins/TraceUtilities/Source/TraceUtilities/TraceUtilities.Build.cs b/Engine/Plugins/TraceUtilities/Source/TraceUtilities/TraceUtilities.Build.cs
--- a/Engine/Plugins/TraceUtilities/Source/TraceUtilities/TraceUtilities.Build.cs
+++ b/Engine/Plugins/TraceUtilities/Source/TraceUtilities/TraceUtilities.Build.cs
@@ -6,24 +6,24 @@ public class TraceUtilities : ModuleRules
 {
 	public TraceUtilities(ReadOnlyTargetRules Target) : base(Target)
 	{
+		UnsafeTypeCastWarningLevel = WarningLevel.Error;
+
 		PCHUsage = ModuleRules.PCHUsageMode.UseExplicitOrSharedPCHs;
 
 		PublicDependencyModuleNames.AddRange(
-			new string[]
-			{
+			new string[] {
 				"Core",
 				// ... add other public dependencies that you statically link with here ...
 			}
-			);
+		);
 
 		PrivateDependencyModuleNames.AddRange(
-			new string[]
-			{
+			new string[] {
 				"CoreUObject",
 				"Engine",
 				"TraceLog",
-				// ... add private dependencies that you statically link with here ...	
+				// ... add private dependencies that you statically link with here ...
 			}
-			);
+		);
 	}
 }
diff --git a/Engine/Source/Developer/TraceAnalysis/Private/Analysis/Engine.cpp b/Engine/Source/Developer/TraceAnalysis/Private/Analysis/Engine.cpp
--- a/Engine/Source/Developer/TraceAnalysis/Private/Analysis/Engine.cpp
+++ b/Engine/Source/Developer/TraceAnalysis/Private/Analysis/Engine.cpp
@@ -1090,7 +1090,7 @@ const FTypeRegistry::FTypeInfo* FTypeRegistry::AddVersion4(const void* TraceData
 	{
 		const auto& Field = NewEvent.Fields[i];
 
-		int8 TypeSize = 1 << (Field.TypeInfo & Protocol0::Field_Pow2SizeMask);
+		int8 TypeSize = (int8)(1 << (Field.TypeInfo & Protocol0::Field_Pow2SizeMask));
 		if (Field.TypeInfo & Protocol0::Field_Float)
 		{
 			TypeSize = -TypeSize;
@@ -1150,7 +1150,7 @@ const FTypeRegistry::FTypeInfo* FTypeRegistry::AddVersion6(const void* TraceData
 		const auto& Field = NewEvent.Fields[i];
 		if (Field.FieldType == EFieldFamily::Regular)
 		{
-			int8 TypeSize = 1 << (Field.Regular.TypeInfo & Protocol0::Field_Pow2SizeMask);
+			int8 TypeSize = (int8)(1 << (Field.Regular.TypeInfo & Protocol0::Field_Pow2SizeMask));
 			if (Field.Regular.TypeInfo & Protocol0::Field_Float)
 			{
 				TypeSize = -TypeSize;
@@ -1168,7 +1168,7 @@ const FTypeRegistry::FTypeInfo* FTypeRegistry::AddVersion6(const void* TraceData
 		else if (Field.FieldType == EFieldFamily::Reference)
 		{
 			check((Field.Reference.TypeInfo & Protocol0::Field_CategoryMask) == Protocol0::Field_Integer);
-			const int8 TypeSize = 1 << (Field.Reference.TypeInfo & Protocol0::Field_Pow2SizeMask);
+			const int8 TypeSize = (int8)(1 << (Field.Reference.TypeInfo & Protocol0::Field_Pow2SizeMask));
 
 			auto& OutField = Builder.AddField(NameCursor, Field.Reference.NameSize, TypeSize);
 			OutField.Offset = Field.Reference.Offset;
@@ -1182,7 +1182,7 @@ const FTypeRegistry::FTypeInfo* FTypeRegistry::AddVersion6(const void* TraceData
 		else if (Field.FieldType == EFieldFamily::DefinitionId)
 		{
 			check((Field.DefinitionId.TypeInfo & Protocol0::Field_CategoryMask) == Protocol0::Field_Integer);
-			const int8 TypeSize = 1 << (Field.DefinitionId.TypeInfo & Protocol0::Field_Pow2SizeMask);
+			const int8 TypeSize = (int8)(1 << (Field.DefinitionId.TypeInfo & Protocol0::Field_Pow2SizeMask));
 
 			auto DefinitionIdFieldName = ANSITEXTVIEW("DefinitionId");
 			auto& OutField = Builder.AddField(DefinitionIdFieldName.GetData(), DefinitionIdFieldName.Len(), TypeSize);
@@ -3115,7 +3115,7 @@ int32 FProtocol5Stage::ParseImportantEvents(FStreamReader& Reader, EventDescArra
 
 		FEventDesc EventDesc;
 		EventDesc.Serial = ESerial::Ignored;
-		EventDesc.Uid = Uid;
+		EventDesc.Uid = (uint16)Uid;
 		EventDesc.Data = Header->Data;
 
 		// Special case for new events. It would work to add a 0 type to the
@@ -3642,7 +3642,7 @@ int32 FProtocol5Stage::ParseEvent(FStreamReader& Reader, FEventDesc& EventDesc,
 	}
 
 	EventDesc.Serial = Serial;
-	EventDesc.Uid = Uid;
+	EventDesc.Uid = (uint16)Uid;
 	EventDesc.Data = Cursor;
 
 	uint32 HeaderSize = uint32(UPTRINT(Cursor - Reader.GetPointer<uint8>()));
diff --git a/Engine/Source/Developer/TraceAnalysis/Private/Store/StoreClient.cpp b/Engine/Source/Developer/TraceAnalysis/Private/Store/StoreClient.cpp
--- a/Engine/Source/Developer/TraceAnalysis/Private/Store/StoreClient.cpp
+++ b/Engine/Source/Developer/TraceAnalysis/Private/Store/StoreClient.cpp
@@ -468,14 +468,14 @@ void FStoreClient::FStatus::GetWatchDirectories(TArray<FString>& OutDirs) const
 uint32 FStoreClient::FVersion::GetMajorVersion() const
 {
 	const auto* Response = (const FResponse*)this;
-	return Response->GetInteger("major", 0);
+	return Response->GetUint32Checked("major", 0);
 }
 
 ////////////////////////////////////////////////////////////////////////////////
 uint32 FStoreClient::FVersion::GetMinorVersion() const
 {
 	const auto* Response = (const FResponse*)this;
-	return Response->GetInteger("minor", 0);
+	return Response->GetUint32Checked("minor", 0);
 }
 
 ////////////////////////////////////////////////////////////////////////////////
diff --git a/Engine/Source/Developer/TraceAnalysis/TraceAnalysis.Build.cs b/Engine/Source/Developer/TraceAnalysis/TraceAnalysis.Build.cs
--- a/Engine/Source/Developer/TraceAnalysis/TraceAnalysis.Build.cs
+++ b/Engine/Source/Developer/TraceAnalysis/TraceAnalysis.Build.cs
@@ -6,13 +6,19 @@ public class TraceAnalysis : ModuleRules
 {
 	public TraceAnalysis(ReadOnlyTargetRules Target) : base(Target)
 	{
+		UnsafeTypeCastWarningLevel = WarningLevel.Error;
+
 		// Asio does not compile with C++20, remove if updated
 		CppStandard = CppStandardVersion.Cpp17;
 
-		PrivateDependencyModuleNames.Add("Asio");
-		PrivateDependencyModuleNames.Add("Cbor");
-		PrivateDependencyModuleNames.Add("Core");
-		PrivateDependencyModuleNames.Add("Sockets");
-		PrivateDependencyModuleNames.Add("TraceLog");
+		PrivateDependencyModuleNames.AddRange(
+			new string[] {
+				"Asio",
+				"Cbor",
+				"Core",
+				"Sockets",
+				"TraceLog",
+			}
+		);
 	}
 }
diff --git a/Engine/Source/Developer/TraceServices/Private/Common/TimelineStatistics.h b/Engine/Source/Developer/TraceServices/Private/Common/TimelineStatistics.h
--- a/Engine/Source/Developer/TraceServices/Private/Common/TimelineStatistics.h
+++ b/Engine/Source/Developer/TraceServices/Private/Common/TimelineStatistics.h
@@ -110,7 +110,7 @@ public:
 
 		// For very large numbers of timers or frames, an out of memory is possible so don't compute the median.
 		constexpr int64 MaxSize = 5 * 100 * 1000 * 1000;
-		bool bComputeMedian = InitialTimerMap.Num() * (int64) FramesNum < MaxSize;
+		bool bComputeMedian = InitialTimerMap.Num() * (int64)FramesNum < MaxSize;
 		for (auto& KV : InitialTimerMap)
 		{
 			FInternalFrameAggregationEntry Entry;
@@ -125,8 +125,8 @@ public:
 
 		TQueue<TSharedPtr<TMap<BucketKeyType, FAggregatedTimingStats>>, EQueueMode::Mpsc> FrameResultsQueue;
 
-		constexpr float RatioOfThreadsToUse = 0.75;
-		int32 NumTasks = FMath::Max(1, int32(GThreadPool->GetNumThreads() * RatioOfThreadsToUse));
+		constexpr float RatioOfThreadsToUse = 0.75f;
+		int32 NumTasks = FMath::Max(1, (int32)((float)GThreadPool->GetNumThreads() * RatioOfThreadsToUse));
 		NumTasks = FMath::Min((int32)NumTasks, Frames.Num());
 		int32 FramesPerTask = Frames.Num() / NumTasks;
 		int32 ExtraFrameTasks = Frames.Num() - NumTasks * FramesPerTask;
diff --git a/Engine/Source/Developer/TraceServices/Private/Model/LoadTimeProfiler.cpp b/Engine/Source/Developer/TraceServices/Private/Model/LoadTimeProfiler.cpp
--- a/Engine/Source/Developer/TraceServices/Private/Model/LoadTimeProfiler.cpp
+++ b/Engine/Source/Developer/TraceServices/Private/Model/LoadTimeProfiler.cpp
@@ -559,7 +559,7 @@ FPackageInfo& FLoadTimeProfilerProvider::CreatePackage()
 {
 	Session.WriteAccessCheck();
 
-	uint32 PackageId = static_cast<uint32>(Packages.Num());
+	uint32 PackageId = uint32(Packages.Num());
 	FPackageInfo& Package = Packages.PushBack();
 	Package.Id = PackageId;
 	return Package;
@@ -586,14 +586,14 @@ void FLoadTimeProfilerProvider::DistributeBytesAcrossFrames(uint64 ByteCount, do
 	double TotalTime = EndTime - StartTime;
 	double TimeInFirstFrame = (BeginFrameIndex + 1) * LoaderFrameLength - StartTime;
 	check(TimeInFirstFrame >= 0.0);
-	uint64 BytesInFirstFrame = uint64(FMath::RoundToZero(TimeInFirstFrame / TotalTime * ByteCount));
+	uint64 BytesInFirstFrame = uint64(FMath::RoundToZero(TimeInFirstFrame / TotalTime * double(ByteCount)));
 	Frames[BeginFrameIndex].*FrameVariable += BytesInFirstFrame;
 
 	uint64 TotalDistributed = BytesInFirstFrame;
 
 	double TimeInLastFrame = EndTime - EndFrameIndex * LoaderFrameLength;
 	check(TimeInLastFrame >= 0.0);
-	uint64 BytesInLastFrame = uint64(FMath::RoundToZero(TimeInLastFrame / TotalTime * ByteCount));
+	uint64 BytesInLastFrame = uint64(FMath::RoundToZero(TimeInLastFrame / TotalTime * double(ByteCount)));
 
 	check(BytesInFirstFrame + BytesInLastFrame <= ByteCount);
 	ByteCount -= BytesInFirstFrame + BytesInLastFrame;
@@ -653,7 +653,7 @@ FPackageExportInfo& FLoadTimeProfilerProvider::CreateExport()
 {
 	Session.WriteAccessCheck();
 
-	uint32 ExportId = static_cast<uint32>(Exports.Num());
+	uint32 ExportId = uint32(Exports.Num());
 	FPackageExportInfo& Export = Exports.PushBack();
 	Export.Id = ExportId;
 	return Export;
diff --git a/Engine/Source/Developer/TraceServices/Private/Model/SymslibResolver.cpp b/Engine/Source/Developer/TraceServices/Private/Model/SymslibResolver.cpp
--- a/Engine/Source/Developer/TraceServices/Private/Model/SymslibResolver.cpp
+++ b/Engine/Source/Developer/TraceServices/Private/Model/SymslibResolver.cpp
@@ -91,13 +91,12 @@ namespace
 		}
 
 		// If the search path is an absolute path to a file use this. Filenames does
-		// not need to to match (e.g. eboot.bin <-> gamename.self)
+		// not need to match (e.g. eboot.bin <-> gamename.self)
 		if (PlatformFile->FileExists(*SearchPath))
 		{
 			return SearchPath;
 		}
 
-
 		// Extract only filename part in case Path is absolute path
 		FString FileName = FPaths::GetCleanFilename(File);
 
@@ -725,8 +724,8 @@ void FSymslibResolver::OnAnalysisComplete()
 			}
 		}
 
-		UE_LOG(LogSymslib, Display, TEXT("Allocated %.02f Mb of strings, %.02f Mb wasted."),
-		       SymbolBytesAllocated / float(1024*1024), SymbolBytesWasted / float(1024*1024));
+		UE_LOG(LogSymslib, Display, TEXT("Allocated %.02f MiB of strings (%.02f MiB wasted)."),
+			(double)SymbolBytesAllocated / (1024.0 * 1024.0), (double)SymbolBytesWasted / (1024.0 * 1024.0));
 	});
 }
 
@@ -805,9 +804,9 @@ void FSymslibResolver::ResolveSymbols(TArrayView<FQueuedAddress>& QueuedWork)
 		}
 		ResolveSymbolTracked(ToResolve.Address, *ToResolve.Target, StringAllocator);
 	}
-	UE_LOG(LogSymslib, VeryVerbose, TEXT("String allocator used: %.02f kb, wasted: %.02f kb using %d blocks"),
-		((StringAllocator.BlockUsed * StringAllocator.BlockSize - StringAllocator.BlockRemaining) * sizeof(TCHAR)) / 1024.0f,
-		(StringAllocator.BlockRemaining * sizeof(TCHAR)) / 1024.0f,
+	UE_LOG(LogSymslib, VeryVerbose, TEXT("String allocator: %.02f KiB used (%.02f KiB wasted) in %d blocks"),
+		(double)((StringAllocator.BlockUsed * StringAllocator.BlockSize - StringAllocator.BlockRemaining) * sizeof(TCHAR)) / 1024.0,
+		(double)(StringAllocator.BlockRemaining * sizeof(TCHAR)) / 1024.0,
 		StringAllocator.BlockUsed);
 	SymbolBytesAllocated.fetch_add(StringAllocator.BlockUsed * StringAllocator.BlockSize * sizeof(TCHAR));
 	SymbolBytesWasted.fetch_add(StringAllocator.BlockRemaining * sizeof(TCHAR));
diff --git a/Engine/Source/Developer/TraceServices/Private/Model/TableMergeTask.h b/Engine/Source/Developer/TraceServices/Private/Model/TableMergeTask.h
--- a/Engine/Source/Developer/TraceServices/Private/Model/TableMergeTask.h
+++ b/Engine/Source/Developer/TraceServices/Private/Model/TableMergeTask.h
@@ -107,10 +107,10 @@ void FTableMergeTask::FillResultUsing(IUntypedTableReader& SourceTableReader) co
 				Row.SetValue(DestIndex, SourceTableReader.GetValueInt(SourceIndex) * Multiplier);
 				break;
 			case TableColumnType_Float:
-				Row.SetValue(DestIndex, SourceTableReader.GetValueFloat(SourceIndex) * Multiplier);
+				Row.SetValue(DestIndex, SourceTableReader.GetValueFloat(SourceIndex) * (float)Multiplier);
 				break;
 			case TableColumnType_Double:
-				Row.SetValue(DestIndex, SourceTableReader.GetValueDouble(SourceIndex) * Multiplier);
+				Row.SetValue(DestIndex, SourceTableReader.GetValueDouble(SourceIndex) * (double)Multiplier);
 				break;
 			case TableColumnType_CString:
 				Row.SetValue(DestIndex, TableC->GetStringStore().Store(SourceTableReader.GetValueCString(SourceIndex)));
diff --git a/Engine/Source/Developer/TraceServices/TraceServices.Build.cs b/Engine/Source/Developer/TraceServices/TraceServices.Build.cs
--- a/Engine/Source/Developer/TraceServices/TraceServices.Build.cs
+++ b/Engine/Source/Developer/TraceServices/TraceServices.Build.cs
@@ -6,12 +6,15 @@ public class TraceServices : ModuleRules
 {
 	public TraceServices(ReadOnlyTargetRules Target) : base(Target)
 	{
+		UnsafeTypeCastWarningLevel = WarningLevel.Error;
+
 		PublicDependencyModuleNames.AddRange(
 			new string[] {
 				"Cbor",
 				"Core",
+				"SymsLib",
 				"TraceAnalysis",
-				"SymsLib"
-			});
+			}
+		);
 	}
 }
diff --git a/Engine/Source/Programs/UnrealInsights/UnrealInsights.Build.cs b/Engine/Source/Programs/UnrealInsights/UnrealInsights.Build.cs
--- a/Engine/Source/Programs/UnrealInsights/UnrealInsights.Build.cs
+++ b/Engine/Source/Programs/UnrealInsights/UnrealInsights.Build.cs
@@ -6,14 +6,15 @@ public class UnrealInsights : ModuleRules
 {
 	public UnrealInsights(ReadOnlyTargetRules Target) : base(Target)
 	{
+		UnsafeTypeCastWarningLevel = WarningLevel.Error;
+
 		PublicIncludePathModuleNames.Add("Launch");
 
 		PrivateDependencyModuleNames.AddRange(
 			new string[] {
 				"AppFramework",
-				//"AutomationController",
-				"Core",
 				"ApplicationCore",
+				"Core",
 				"CoreUObject",
 				"Projects",
 				"Slate",
