[Insights]
  - TraceAnalysis: Added debug functionality for trace transport and trace protocol analysis. See TraceAnalysisDebug.h for the available debug options (UE_TRACE_ANALYSIS_DEBUG_API, UE_TRACE_ANALYSIS_DEBUG, UE_TRACE_ANALYSIS_DEBUG_LOG, UE_TRACE_ANALYSIS_DEBUG_LOG_IMPL, UE_TRACE_ANALYSIS_DEBUG_LEVEL).
  - TraceAnalysis: Changed the limit for accumulated too much data on a single thread from 512 MiB to 2000 MB. Also implemented better code to skip serials when limit is encountered.
  - TraceAnalysis: Implemented workaround for corrupted traces (events with aux blocks missing the "aux terminal").
  - TraceAnalysis: Fixed EventData.GetArray<>() (to be safe to be called multiple times).

diff --git a/Engine/Source/Developer/TraceAnalysis/Private/Analysis/Engine.cpp b/Engine/Source/Developer/TraceAnalysis/Private/Analysis/Engine.cpp
--- a/Engine/Source/Developer/TraceAnalysis/Private/Analysis/Engine.cpp
+++ b/Engine/Source/Developer/TraceAnalysis/Private/Analysis/Engine.cpp
@@ -10,8 +10,10 @@
 #include "HAL/UnrealMemory.h"
 #include "Internationalization/Internationalization.h"
 #include "Logging/LogMacros.h"
+#include "Misc/StringBuilder.h"
 #include "StreamReader.h"
 #include "Templates/UnrealTemplate.h"
+#include "TraceAnalysisDebug.h"
 #include "Trace/Analysis.h"
 #include "Trace/Analyzer.h"
 #include "Trace/Detail/Protocol.h"
@@ -22,13 +24,6 @@
 
 DEFINE_LOG_CATEGORY_STATIC(LogTraceAnalysis, Log, All)
 
-#if 0
-#	define TRACE_ANALYSIS_DEBUG(Format, ...) \
-		do { UE_LOG(LogTraceAnalysis, Log, TEXT(Format), __VA_ARGS__) } while (0)
-#else
-#	define TRACE_ANALYSIS_DEBUG(...)
-#endif
-
 #define LOCTEXT_NAMESPACE "TraceAnalysis"
 
 namespace UE {
@@ -149,6 +144,7 @@ void FAuxDataCollector::Defragment(FAuxData& Data)
 	Data.Data = Buffer;
 	Data.DataSize = Size;
 	Data.bOwnsData = 1;
+	Data.bFragmented = 0;
 }
 
 ////////////////////////////////////////////////////////////////////////////////
@@ -262,7 +258,6 @@ const TArray<uint8>* FThreads::GetGroupName() const
 
 
 
-
 // {{{1 thread-info ------------------------------------------------------------
 
 ////////////////////////////////////////////////////////////////////////////////
@@ -384,7 +379,7 @@ struct FDispatch
 		int8		SizeAndType;		// value == byte_size, sign == float < 0 < int
 		uint8		Class : 7;
 		uint8		bArray : 1;
-		uint32		RefUid;			// If reference field, uuid of ref type
+		uint32		RefUid;				// If reference field, uid of ref type
 	};
 
 	int32			GetFieldIndex(const ANSICHAR* Name) const;
@@ -574,6 +569,15 @@ uint32 IAnalyzer::FEventTypeInfo::GetId() const
 	return Inner->Uid;
 }
 
+////////////////////////////////////////////////////////////////////////////////
+#if UE_TRACE_ANALYSIS_DEBUG_API
+uint8 IAnalyzer::FEventTypeInfo::GetFlags() const
+{
+	const auto* Inner = (const FDispatch*)this;
+	return Inner->Flags;
+}
+#endif // UE_TRACE_ANALYSIS_DEBUG_API
+
 ////////////////////////////////////////////////////////////////////////////////
 const ANSICHAR* IAnalyzer::FEventTypeInfo::GetName() const
 {
@@ -614,6 +618,13 @@ const IAnalyzer::FEventFieldInfo* IAnalyzer::FEventTypeInfo::GetFieldInfo(uint32
 	return (const IAnalyzer::FEventFieldInfo*)(Inner->Fields + Index);
 }
 
+////////////////////////////////////////////////////////////////////////////////
+int32 IAnalyzer::FEventTypeInfo::GetFieldIndex(const ANSICHAR* FieldName) const
+{
+	const auto* Inner = (const FDispatch*)this;
+	return Inner->GetFieldIndex(FieldName);
+}
+
 ////////////////////////////////////////////////////////////////////////////////
 IAnalyzer::FEventFieldHandle IAnalyzer::FEventTypeInfo::GetFieldHandleUnchecked(uint32 Index) const
 {
@@ -622,7 +633,7 @@ IAnalyzer::FEventFieldHandle IAnalyzer::FEventTypeInfo::GetFieldHandleUnchecked(
 	{
 		return FEventFieldHandle { -1 };
 	}
-	return FEventFieldHandle {Inner->Fields[Index].Offset };
+	return FEventFieldHandle { Inner->Fields[Index].Offset };
 }
 
 ////////////////////////////////////////////////////////////////////////////////
@@ -702,6 +713,15 @@ bool IAnalyzer::FEventFieldInfo::IsSigned() const
 	return (Inner->Class & UE::Trace::Protocol0::Field_Signed) != 0;
 }
 
+////////////////////////////////////////////////////////////////////////////////
+#if UE_TRACE_ANALYSIS_DEBUG_API
+uint32 IAnalyzer::FEventFieldInfo::GetOffset() const
+{
+	const auto* Inner = (const FDispatch::FField*)this;
+	return Inner->Offset;
+}
+#endif // UE_TRACE_ANALYSIS_DEBUG_API
+
 ////////////////////////////////////////////////////////////////////////////////
 uint8 IAnalyzer::FEventFieldInfo::GetSize() const
 {
@@ -709,6 +729,8 @@ uint8 IAnalyzer::FEventFieldInfo::GetSize() const
 	return uint8(Inner->Size);
 }
 
+
+
 // {{{1 array-reader -----------------------------------------------------------
 
 ////////////////////////////////////////////////////////////////////////////////
@@ -720,6 +742,33 @@ uint32 IAnalyzer::FArrayReader::Num() const
 	return (SizeAndType == 0) ? SizeAndType : (Inner->DataSize / SizeAndType);
 }
 
+////////////////////////////////////////////////////////////////////////////////
+#if UE_TRACE_ANALYSIS_DEBUG_API
+const uint8* IAnalyzer::FArrayReader::GetRawData() const
+{
+	const auto* Inner = (const FAuxData*)this;
+	return Inner->Data;
+}
+#endif // UE_TRACE_ANALYSIS_DEBUG_API
+
+////////////////////////////////////////////////////////////////////////////////
+#if UE_TRACE_ANALYSIS_DEBUG_API
+uint32 IAnalyzer::FArrayReader::GetRawDataSize() const
+{
+	const auto* Inner = (const FAuxData*)this;
+	return Inner->DataSize;
+}
+#endif // UE_TRACE_ANALYSIS_DEBUG_API
+
+////////////////////////////////////////////////////////////////////////////////
+#if UE_TRACE_ANALYSIS_DEBUG_API
+int8 IAnalyzer::FArrayReader::GetSizeAndType() const
+{
+	const auto* Inner = (const FAuxData*)this;
+	return Inner->FieldSizeAndType;
+}
+#endif // UE_TRACE_ANALYSIS_DEBUG_API
+
 ////////////////////////////////////////////////////////////////////////////////
 const void* IAnalyzer::FArrayReader::GetImpl(uint32 Index, int8& SizeAndType) const
 {
@@ -968,6 +1017,131 @@ uint32 IAnalyzer::FEventData::GetAttachmentSize() const
 	return Info->Size - Info->Dispatch.EventSize;
 }
 
+////////////////////////////////////////////////////////////////////////////////
+#if UE_TRACE_ANALYSIS_DEBUG_API
+const uint8* IAnalyzer::FEventData::GetRawPointer() const
+{
+	const auto* Info = (const FEventDataInfo*)this;
+	return Info->Ptr;
+}
+#endif // UE_TRACE_ANALYSIS_DEBUG_API
+
+////////////////////////////////////////////////////////////////////////////////
+#if UE_TRACE_ANALYSIS_DEBUG_API
+uint32 IAnalyzer::FEventData::GetRawSize() const
+{
+	const auto* Info = (const FEventDataInfo*)this;
+	return Info->Size;
+}
+#endif // UE_TRACE_ANALYSIS_DEBUG_API
+
+////////////////////////////////////////////////////////////////////////////////
+#if UE_TRACE_ANALYSIS_DEBUG_API
+uint32 IAnalyzer::FEventData::GetAuxSize() const
+{
+	const auto* Info = (const FEventDataInfo*)this;
+
+	if ((Info->Dispatch.Flags & UE::Trace::FDispatch::Flag_MaybeHasAux) == 0)
+	{
+		return 0;
+	}
+	else
+	{
+		uint32 Size = 0;
+		if (Info->AuxCollector)
+		{
+			for (FAuxData& Data : *(Info->AuxCollector))
+			{
+				Size += 4; // aux field header
+				Size += Data.DataSize;
+			}
+		}
+		++Size; // aux terminator
+		return Size;
+	}
+}
+#endif // UE_TRACE_ANALYSIS_DEBUG_API
+
+////////////////////////////////////////////////////////////////////////////////
+#if UE_TRACE_ANALYSIS_DEBUG_API
+uint32 IAnalyzer::FEventData::GetTotalSize(IAnalyzer::EStyle Style, const IAnalyzer::FOnEventContext& Context, uint32 ProtocolVersion) const
+{
+	uint32 Size = 0;
+
+	if (Style == EStyle::EnterScope)
+	{
+		// Include the size of the previous EnterScope event.
+		if (Context.EventTime.GetTimestamp() == 0)
+		{
+			Size = 1; // uid
+		}
+		else
+		{
+			Size = 8; // 1 uid + 7 timestamp
+		}
+	}
+	else if (Style == EStyle::LeaveScope)
+	{
+		// Return only the size of the LeaveScope event.
+		if (Context.EventTime.GetTimestamp() == 0)
+		{
+			return 1; // uid
+		}
+		else
+		{
+			return 8; // 1 uid + 7 timestamp
+		}
+	}
+
+	const auto* Info = (const FEventDataInfo*)this;
+
+	const uint16 KnownEventUids = (ProtocolVersion >= 5) ? Protocol5::EKnownEventUids::User :
+								  (ProtocolVersion == 4) ? Protocol4::EKnownEventUids::User : 0;
+	static_assert(Protocol6::EKnownEventUids::User == Protocol5::EKnownEventUids::User, "Protocol6::EKnownEventUids::User");
+	static_assert(Protocol7::EKnownEventUids::User == Protocol5::EKnownEventUids::User, "Protocol7::EKnownEventUids::User");
+
+	// Add header size.
+	if (Info->Dispatch.Uid < KnownEventUids)
+	{
+		Size += 1; // uid
+	}
+	else if ((Info->Dispatch.Flags & UE::Trace::FDispatch::Flag_Important) != 0)
+	{
+		Size += 4; // sizeof(UE::Trace::Protocol5::FImportantEventHeader);
+	}
+	else if (Info->Dispatch.Flags & UE::Trace::FDispatch::Flag_NoSync)
+	{
+		if (ProtocolVersion < 5)
+		{
+			Size += 4; // sizeof(UE::Trace::Protocol2::FEventHeader);
+		}
+		else
+		{
+			Size += 2; // uid == sizeof(UE::Trace::Protocol5::FEventHeader);
+		}
+	}
+	else // Sync
+	{
+		if (ProtocolVersion < 5)
+		{
+			Size += 7; // sizeof(UE::Trace::Protocol2::FEventHeaderSync);
+		}
+		else
+		{
+			Size += 2 + 3; // uid + serial == sizeof(UE::Trace::Protocol5::FEventHeaderSync);
+		}
+	}
+
+	// Add fixed size and attachment size.
+	Size += Info->Size;
+
+	// Add aux data size (including headers).
+	Size += GetAuxSize();
+
+	return Size;
+}
+#endif // UE_TRACE_ANALYSIS_DEBUG_API
+
 ////////////////////////////////////////////////////////////////////////////////
 bool IAnalyzer::FEventData::IsDefinitionImpl(uint32& OutTypeId) const
 {
@@ -1235,12 +1409,19 @@ void FTypeRegistry::Add(FTypeInfo* TypeInfo)
 	{
 		if (TypeInfos[Uid] != nullptr)
 		{
+			UE_TRACE_ANALYSIS_DEBUG_LOG("Warning: Override type for Uid=%u", uint32(Uid));
 			FMemory::Free(TypeInfos[Uid]);
 			TypeInfos[Uid] = nullptr;
 		}
 	}
 	else
 	{
+#if UE_TRACE_ANALYSIS_DEBUG
+		if (Uid > TypeInfos.Num() + 100)
+		{
+			UE_TRACE_ANALYSIS_DEBUG_LOG("Warning: Unexpected large Uid=%u", uint32(Uid));
+		}
+#endif // UE_TRACE_ANALYSIS_DEBUG
 		TypeInfos.SetNum(Uid + 1);
 	}
 
@@ -1263,6 +1444,9 @@ public:
 	void				OnNewType(const FTypeRegistry::FTypeInfo* TypeInfo);
 	void				OnEvent(const FTypeRegistry::FTypeInfo& TypeInfo, IAnalyzer::EStyle Style, const IAnalyzer::FOnEventContext& Context);
 	void				OnThreadInfo(const FThreads::FInfo& ThreadInfo);
+#if UE_TRACE_ANALYSIS_DEBUG_API
+	void				OnVersion(uint32 TransportVersion, uint32 ProtocolVersion);
+#endif // UE_TRACE_ANALYSIS_DEBUG_API
 
 private:
 	void				BuildRoutes();
@@ -1399,6 +1583,22 @@ void FAnalyzerHub::OnNewType(const FTypeRegistry::FTypeInfo* TypeInfo)
 
 	TypeToRoute[Uid] = uint16(FirstRoute + 1);
 
+#if UE_TRACE_ANALYSIS_DEBUG
+	const IAnalyzer::FEventTypeInfo& EventTypeInfo = *(IAnalyzer::FEventTypeInfo*)TypeInfo;
+#if UE_TRACE_ANALYSIS_DEBUG_API
+	const uint8 EventTypeFlags = EventTypeInfo.GetFlags();
+	UE_TRACE_ANALYSIS_DEBUG_LOG("NewEvent: Uid=%u (%s.%s) Flags=%s%s%s%s",
+		EventTypeInfo.GetId(), EventTypeInfo.GetLoggerName(), EventTypeInfo.GetName(),
+		(EventTypeFlags & FDispatch::Flag_Definition) ? "Definition|" : "",
+		(EventTypeFlags & FDispatch::Flag_Important) ? "Important|" : "",
+		(EventTypeFlags & FDispatch::Flag_MaybeHasAux) ? "MaybeHasAux|" : "",
+		(EventTypeFlags & FDispatch::Flag_NoSync) ? "NoSync" : "Sync");
+#else // UE_TRACE_ANALYSIS_DEBUG_API
+	UE_TRACE_ANALYSIS_DEBUG_LOG("NewEvent: Uid=%u (%s.%s)",
+		EventTypeInfo.GetId(), EventTypeInfo.GetLoggerName(), EventTypeInfo.GetName());
+#endif // UE_TRACE_ANALYSIS_DEBUG_API
+#endif // UE_TRACE_ANALYSIS_DEBUG
+
 	// Inform routes that a new event has been declared.
 	if (FirstRoute >= 0)
 	{
@@ -1447,6 +1647,20 @@ void FAnalyzerHub::OnThreadInfo(const FThreads::FInfo& ThreadInfo)
 	}
 }
 
+////////////////////////////////////////////////////////////////////////////////
+#if UE_TRACE_ANALYSIS_DEBUG_API
+void FAnalyzerHub::OnVersion(uint32 TransportVersion, uint32 ProtocolVersion)
+{
+	for (IAnalyzer* Analyzer : Analyzers)
+	{
+		if (Analyzer != nullptr)
+		{
+			Analyzer->OnVersion(TransportVersion, ProtocolVersion);
+		}
+	}
+}
+#endif // UE_TRACE_ANALYSIS_DEBUG_API
+
 ////////////////////////////////////////////////////////////////////////////////
 void FAnalyzerHub::BuildRoutes()
 {
@@ -1591,6 +1805,33 @@ struct FAnalysisState
 	FTiming		Timing;
 	FSerial		Serial;
 	uint32		UserUidBias = Protocol4::EKnownEventUids::User;
+
+#if UE_TRACE_ANALYSIS_DEBUG
+	uint32 TempUidBytes = 0; // for protocol 4
+	int32 MaxEventDescs = 0; // for protocol 5
+	int32 SerialWrappedCount = 0; // for protocol 5
+	int32 NumSkippedSerialGaps = 0; // for protocol 5
+
+	uint64 TotalEventCount = 0;
+	uint64 NewEventCount = 0;
+	uint64 SyncEventCount = 0;
+	uint64 ImportantNoSyncEventCount = 0;
+	uint64 OtherNoSyncEventCount = 0;
+	uint64 EnterScopeEventCount = 0;
+	uint64 LeaveScopeEventCount = 0;
+	uint64 EnterScopeTEventCount = 0;
+	uint64 LeaveScopeTEventCount = 0;
+
+	uint64 TotalEventSize = 0;
+	uint64 NewEventSize = 0;
+	uint64 SyncEventSize = 0;
+	uint64 ImportantNoSyncEventSize = 0;
+	uint64 OtherNoSyncEventSize = 0;
+	uint64 EnterScopeEventSize = 0;
+	uint64 LeaveScopeEventSize = 0;
+	uint64 EnterScopeTEventSize = 0;
+	uint64 LeaveScopeTEventSize = 0;
+#endif // UE_TRACE_ANALYSIS_DEBUG
 };
 
 
@@ -1687,20 +1928,25 @@ void FTraceAnalyzer::OnNewTrace(const FOnEventContext& Context)
 	// mark the MSB to indicate that the current serial should be corrected.
 	auto& Serial = State.Serial;
 	uint32 Hint = EventData.GetValue<uint32>("Serial");
+	UE_TRACE_ANALYSIS_DEBUG_LOG("Serial=%u (hint)", Hint);
 	Hint -= (Serial.Mask + 1) >> 1;
 	Hint &= Serial.Mask;
 	Serial.Value = Hint;
 	Serial.Value |= 0xc0000000;
+	UE_TRACE_ANALYSIS_DEBUG_LOG("Serial.Value=%u|0xC0000000 (Next=%u)", Hint, Serial.Value & Serial.Mask);
+	UE_TRACE_ANALYSIS_DEBUG_LOG("Serial.Mask=0x%X", Serial.Mask);
 
 	// Later traces will have an explicit "SerialSync" trace event to indicate
 	// when there is enough data to establish the correct log serial
 	const uint8 FeatureSet = EventData.GetValue<uint8>("FeatureSet");
+	UE_TRACE_ANALYSIS_DEBUG_LOG("FeatureSet=%d", FeatureSet);
 	if ((FeatureSet & 1) == 0)
 	{
 		OnSerialSync(Context);
 	}
 
 	State.UserUidBias = EventData.GetValue<uint32>("UserUidBias", uint32(UE::Trace::Protocol3::EKnownEventUids::User));
+	UE_TRACE_ANALYSIS_DEBUG_LOG("UserUidBias=%u", State.UserUidBias);
 
 	OnTiming(Context);
 }
@@ -1709,6 +1955,7 @@ void FTraceAnalyzer::OnNewTrace(const FOnEventContext& Context)
 void FTraceAnalyzer::OnSerialSync(const FOnEventContext& Context)
 {
 	State.Serial.Value &= ~0x40000000;
+	UE_TRACE_ANALYSIS_DEBUG_LOG("SerialSync: Next=%u", State.Serial.Value & State.Serial.Mask);
 }
 
 ////////////////////////////////////////////////////////////////////////////////
@@ -1777,6 +2024,8 @@ void FTraceAnalyzer::OnThreadInfoInternal(const FOnEventContext& Context)
 		}
 	}
 
+	UE_TRACE_ANALYSIS_DEBUG_LOG("Thread: Tid=%u SysId=0x%X Name=\"%s\"", ThreadInfo->ThreadId, ThreadInfo->SystemId, ThreadInfo->Name.GetData());
+
 	ThreadInfoCallback.OnThreadInfo(*ThreadInfo);
 }
 
@@ -1796,7 +2045,11 @@ void FTraceAnalyzer::OnThreadGroupEnd()
 	State.Threads.SetGroupName("", 0);
 }
 
+
+
 // {{{1 bridge -----------------------------------------------------------------
+
+////////////////////////////////////////////////////////////////////////////////
 class FAnalysisBridge
 	: public FThreadInfoCallback
 {
@@ -1806,8 +2059,10 @@ public:
 						FAnalysisBridge(TArray<IAnalyzer*>&& Analyzers);
 	bool				IsStillAnalyzing() const;
 	void				Reset();
+	FAnalysisState&		GetState();
 	uint32				GetUserUidBias() const;
 	FSerial&			GetSerial();
+	uint32				GetActiveThreadId() const;
 	void				SetActiveThread(uint32 ThreadId);
 	void				OnNewType(const FTypeRegistry::FTypeInfo* TypeInfo);
 	void				OnEvent(const FEventDataInfo& EventDataInfo);
@@ -1820,6 +2075,21 @@ public:
 	void				LeaveScope(uint64 RelativeTimestamp);
 	void				LeaveScopeA(uint64 AbsoluteTimestamp);
 	void				LeaveScopeB(uint64 BaseRelativeTimestamp);
+#if UE_TRACE_ANALYSIS_DEBUG_API
+	void				OnVersion(uint32 TransportVersion, uint32 ProtocolVersion) { AnalyzerHub.OnVersion(TransportVersion, ProtocolVersion); }
+#endif // UE_TRACE_ANALYSIS_DEBUG_API
+#if UE_TRACE_ANALYSIS_DEBUG
+	void				DebugLogEvent(const FTypeRegistry::FTypeInfo* TypeInfo, uint32 FixedSize, uint32 AuxSize, uint32 Serial);
+	void				DebugLogNewEvent(uint32 Uid, const FTypeRegistry::FTypeInfo* TypeInfo, uint32 EventSize);
+	void				DebugLogEnterScopeEvent(uint32 Uid, uint32 EventSize);
+	void				DebugLogEnterScopeEvent(uint32 Uid, uint64 RelativeTimestamp, uint32 EventSize);
+	void				DebugLogEnterScopeAEvent(uint32 Uid, uint64 AbsoluteTimestamp, uint32 EventSize);
+	void				DebugLogEnterScopeBEvent(uint32 Uid, uint64 BaseRelativeTimestamp, uint32 EventSize);
+	void				DebugLogLeaveScopeEvent(uint32 Uid, uint32 EventSize);
+	void				DebugLogLeaveScopeEvent(uint32 Uid, uint64 RelativeTimestamp, uint32 EventSize);
+	void				DebugLogLeaveScopeAEvent(uint32 Uid, uint64 AbsoluteTimestamp, uint32 EventSize);
+	void				DebugLogLeaveScopeBEvent(uint32 Uid, uint64 BaseRelativeTimestamp, uint32 EventSize);
+#endif // UE_TRACE_ANALYSIS_DEBUG_API
 
 private:
 	void				DispatchLeaveScope();
@@ -1854,9 +2124,9 @@ void FAnalysisBridge::Reset()
 }
 
 ////////////////////////////////////////////////////////////////////////////////
-void FAnalysisBridge::SetActiveThread(uint32 ThreadId)
+FAnalysisState& FAnalysisBridge::GetState()
 {
-	ThreadInfo = State.Threads.GetInfo(ThreadId);
+	return State;
 }
 
 ////////////////////////////////////////////////////////////////////////////////
@@ -1871,6 +2141,18 @@ FAnalysisBridge::FSerial& FAnalysisBridge::GetSerial()
 	return State.Serial;
 }
 
+////////////////////////////////////////////////////////////////////////////////
+uint32 FAnalysisBridge::GetActiveThreadId() const
+{
+	return ThreadInfo ? ThreadInfo->ThreadId : ~0u;
+}
+
+////////////////////////////////////////////////////////////////////////////////
+void FAnalysisBridge::SetActiveThread(uint32 ThreadId)
+{
+	ThreadInfo = State.Threads.GetInfo(ThreadId);
+}
+
 ////////////////////////////////////////////////////////////////////////////////
 void FAnalysisBridge::OnNewType(const FTypeRegistry::FTypeInfo* TypeInfo)
 {
@@ -2004,9 +2286,198 @@ void FAnalysisBridge::DispatchLeaveScope()
 	AnalyzerHub.OnEvent(*TypeInfo, IAnalyzer::EStyle::LeaveScope, Context);
 }
 
+#if UE_TRACE_ANALYSIS_DEBUG
+////////////////////////////////////////////////////////////////////////////////
+void FAnalysisBridge::DebugLogEvent(const FTypeRegistry::FTypeInfo* TypeInfo, uint32 FixedSize, uint32 AuxSize, uint32 Serial)
+{
+	++State.TotalEventCount;
+	uint32 EventSize = FixedSize + AuxSize;
+	State.TotalEventSize += EventSize;
+
+#if UE_TRACE_ANALYSIS_DEBUG_LEVEL >= 2
+	UE_TRACE_ANALYSIS_DEBUG_BeginStringBuilder();
+	UE_TRACE_ANALYSIS_DEBUG_Appendf("[EVENT %llu]", State.TotalEventCount);
+
+	if (GetActiveThreadId() != ~0u)
+	{
+		UE_TRACE_ANALYSIS_DEBUG_Appendf(" Tid=%u", GetActiveThreadId());
+	}
+
+	UE_TRACE_ANALYSIS_DEBUG_Appendf(" Uid=%u", uint32(TypeInfo->Uid));
+
+	const IAnalyzer::FEventTypeInfo& EventTypeInfo = *(IAnalyzer::FEventTypeInfo*)TypeInfo;
+	UE_TRACE_ANALYSIS_DEBUG_Appendf(" (%s.%s)", EventTypeInfo.GetLoggerName(), EventTypeInfo.GetName());
+
+	UE_TRACE_ANALYSIS_DEBUG_Appendf(" Size=%u", EventSize);
+	if (AuxSize != 0)
+	{
+		UE_TRACE_ANALYSIS_DEBUG_Appendf(" (%u + %u)", FixedSize, AuxSize);
+	}
+#endif // UE_TRACE_ANALYSIS_DEBUG_LEVEL
+
+	if ((TypeInfo->Flags & FTypeRegistry::FTypeInfo::Flag_NoSync) == 0)
+	{
+		++State.SyncEventCount;
+		State.SyncEventSize += EventSize;
+#if UE_TRACE_ANALYSIS_DEBUG_LEVEL >= 2
+		UE_TRACE_ANALYSIS_DEBUG_Appendf(" Sync=%u", Serial);
+#endif // UE_TRACE_ANALYSIS_DEBUG_LEVEL
+	}
+	else if ((TypeInfo->Flags & FTypeRegistry::FTypeInfo::Flag_Important) != 0)
+	{
+		++State.ImportantNoSyncEventCount;
+		State.ImportantNoSyncEventSize += EventSize;
+#if UE_TRACE_ANALYSIS_DEBUG_LEVEL >= 2
+		UE_TRACE_ANALYSIS_DEBUG_Append(" Important");
+#endif // UE_TRACE_ANALYSIS_DEBUG_LEVEL
+	}
+	else
+	{
+		++State.OtherNoSyncEventCount;
+		State.OtherNoSyncEventSize += EventSize;
+	}
+
+#if UE_TRACE_ANALYSIS_DEBUG_LEVEL >= 2
+	UE_TRACE_ANALYSIS_DEBUG_EndStringBuilder();
+#endif // UE_TRACE_ANALYSIS_DEBUG_LEVEL
+}
+
+////////////////////////////////////////////////////////////////////////////////
+void FAnalysisBridge::DebugLogNewEvent(uint32 Uid, const FTypeRegistry::FTypeInfo* TypeInfo, uint32 EventSize)
+{
+	++State.TotalEventCount;
+	++State.NewEventCount;
+	State.TotalEventSize += EventSize;
+	State.NewEventSize += EventSize;
+#if UE_TRACE_ANALYSIS_DEBUG_LEVEL >= 2
+	const uint32 Tid = ETransportTid::Importants; // GetActiveThreadId()
+	UE_TRACE_ANALYSIS_DEBUG_LOG("[EVENT %llu] Tid=%u Uid=%u (NewEvent) Size=%u", State.TotalEventCount, Tid, Uid, EventSize);
+#endif // UE_TRACE_ANALYSIS_DEBUG_LEVEL
+}
+
+////////////////////////////////////////////////////////////////////////////////
+void FAnalysisBridge::DebugLogEnterScopeEvent(uint32 Uid, uint32 EventSize)
+{
+	++State.TotalEventCount;
+	++State.EnterScopeEventCount;
+	State.TotalEventSize += EventSize;
+	State.EnterScopeEventSize += EventSize;
+#if UE_TRACE_ANALYSIS_DEBUG_LEVEL >= 2
+	UE_TRACE_ANALYSIS_DEBUG_LOG("[EVENT %llu] Tid=%u Uid=%u (EnterScope) Size=%u",
+		State.TotalEventCount, GetActiveThreadId(), Uid, EventSize);
+#endif // UE_TRACE_ANALYSIS_DEBUG_LEVEL
+}
+
+////////////////////////////////////////////////////////////////////////////////
+void FAnalysisBridge::DebugLogEnterScopeEvent(uint32 Uid, uint64 RelativeTimestamp, uint32 EventSize)
+{
+	++State.TotalEventCount;
+	++State.EnterScopeTEventCount;
+	State.TotalEventSize += EventSize;
+	State.EnterScopeTEventSize += EventSize;
+#if UE_TRACE_ANALYSIS_DEBUG_LEVEL >= 2
+	check(ThreadInfo);
+	UE_TRACE_ANALYSIS_DEBUG_LOG("[EVENT %llu] Tid=%u Uid=%u (EnterScope_T) Timestamp=(+%llu)=%llu Size=%u",
+		State.TotalEventCount, GetActiveThreadId(), Uid,
+		RelativeTimestamp, State.Timing.BaseTimestamp + ThreadInfo->PrevTimestamp + RelativeTimestamp,
+		EventSize);
+#endif // UE_TRACE_ANALYSIS_DEBUG_LEVEL
+}
+
+////////////////////////////////////////////////////////////////////////////////
+void FAnalysisBridge::DebugLogEnterScopeAEvent(uint32 Uid, uint64 AbsoluteTimestamp, uint32 EventSize)
+{
+	++State.TotalEventCount;
+	++State.EnterScopeTEventCount;
+	State.TotalEventSize += EventSize;
+	State.EnterScopeTEventSize += EventSize;
+#if UE_TRACE_ANALYSIS_DEBUG_LEVEL >= 2
+	UE_TRACE_ANALYSIS_DEBUG_LOG("[EVENT %llu] Tid=%u Uid=%u (EnterScope_TA) Timestamp=%llu Size=%u",
+		State.TotalEventCount, GetActiveThreadId(), Uid,
+		AbsoluteTimestamp,
+		EventSize);
+#endif // UE_TRACE_ANALYSIS_DEBUG_LEVEL
+}
+
+////////////////////////////////////////////////////////////////////////////////
+void FAnalysisBridge::DebugLogEnterScopeBEvent(uint32 Uid, uint64 BaseAbsoluteTimestamp, uint32 EventSize)
+{
+	++State.TotalEventCount;
+	++State.EnterScopeTEventCount;
+	State.TotalEventSize += EventSize;
+	State.EnterScopeTEventSize += EventSize;
+#if UE_TRACE_ANALYSIS_DEBUG_LEVEL >= 2
+	UE_TRACE_ANALYSIS_DEBUG_LOG("[EVENT %llu] Tid=%u Uid=%u (EnterScope_TB) Timestamp=(+%llu)=%llu Size=%u",
+		State.TotalEventCount, GetActiveThreadId(), Uid,
+		BaseAbsoluteTimestamp, State.Timing.BaseTimestamp + BaseAbsoluteTimestamp,
+		EventSize);
+#endif // UE_TRACE_ANALYSIS_DEBUG_LEVEL
+}
+
+////////////////////////////////////////////////////////////////////////////////
+void FAnalysisBridge::DebugLogLeaveScopeEvent(uint32 Uid, uint32 EventSize)
+{
+	++State.TotalEventCount;
+	++State.LeaveScopeEventCount;
+	State.TotalEventSize += EventSize;
+	State.LeaveScopeEventSize += EventSize;
+#if UE_TRACE_ANALYSIS_DEBUG_LEVEL >= 2
+	UE_TRACE_ANALYSIS_DEBUG_LOG("[EVENT %llu] Tid=%u Uid=%u (LeaveScope) Size=%u",
+		State.TotalEventCount, GetActiveThreadId(), Uid, EventSize);
+#endif // UE_TRACE_ANALYSIS_DEBUG_LEVEL
+}
+
+////////////////////////////////////////////////////////////////////////////////
+void FAnalysisBridge::DebugLogLeaveScopeEvent(uint32 Uid, uint64 RelativeTimestamp, uint32 EventSize)
+{
+	++State.TotalEventCount;
+	++State.LeaveScopeTEventCount;
+	State.TotalEventSize += EventSize;
+	State.LeaveScopeTEventSize += EventSize;
+#if UE_TRACE_ANALYSIS_DEBUG_LEVEL >= 2
+	check(ThreadInfo);
+	UE_TRACE_ANALYSIS_DEBUG_LOG("[EVENT %llu] Tid=%u Uid=%u (LeaveScope_T) Timestamp=(+%llu)=%llu Size=%u",
+		State.TotalEventCount, GetActiveThreadId(), Uid,
+		RelativeTimestamp, State.Timing.BaseTimestamp + ThreadInfo->PrevTimestamp + RelativeTimestamp,
+		EventSize);
+#endif // UE_TRACE_ANALYSIS_DEBUG_LEVEL
+}
+
+////////////////////////////////////////////////////////////////////////////////
+void FAnalysisBridge::DebugLogLeaveScopeAEvent(uint32 Uid, uint64 AbsoluteTimestamp, uint32 EventSize)
+{
+	++State.TotalEventCount;
+	++State.LeaveScopeTEventCount;
+	State.TotalEventSize += EventSize;
+	State.LeaveScopeTEventSize += EventSize;
+#if UE_TRACE_ANALYSIS_DEBUG_LEVEL >= 2
+	UE_TRACE_ANALYSIS_DEBUG_LOG("[EVENT %llu] Tid=%u Uid=%u (LeaveScope_TA) Timestamp=%llu Size=%u",
+		State.TotalEventCount, GetActiveThreadId(), Uid,
+		AbsoluteTimestamp,
+		EventSize);
+#endif // UE_TRACE_ANALYSIS_DEBUG_LEVEL
+}
+
+////////////////////////////////////////////////////////////////////////////////
+void FAnalysisBridge::DebugLogLeaveScopeBEvent(uint32 Uid, uint64 BaseAbsoluteTimestamp, uint32 EventSize)
+{
+	++State.TotalEventCount;
+	++State.LeaveScopeTEventCount;
+	State.TotalEventSize += EventSize;
+	State.LeaveScopeTEventSize += EventSize;
+#if UE_TRACE_ANALYSIS_DEBUG_LEVEL >= 2
+	UE_TRACE_ANALYSIS_DEBUG_LOG("[EVENT %llu] Tid=%u Uid=%u (LeaveScope_TB) Timestamp=(+%llu)=%llu Size=%u",
+		State.TotalEventCount, GetActiveThreadId(), Uid,
+		BaseAbsoluteTimestamp, State.Timing.BaseTimestamp + BaseAbsoluteTimestamp,
+		EventSize);
+#endif // UE_TRACE_ANALYSIS_DEBUG_LEVEL
+}
+#endif // UE_TRACE_ANALYSIS_DEBUG
+
 
 
 // {{{1 machine ----------------------------------------------------------------
+
 ////////////////////////////////////////////////////////////////////////////////
 class FAnalysisMachine
 {
@@ -2035,7 +2506,7 @@ public:
 		template<typename FormatType, typename... Types>
 		inline void EmitMessagef(EAnalysisMessageSeverity Severity, const FormatType& Format, Types... Args) const
 		{
-			TStringBuilder<64> FormattedMessage;
+			TStringBuilder<128> FormattedMessage;
 			FormattedMessage.Appendf(Format, Forward<Types>(Args)...);
 			EmitMessage(Severity, FormattedMessage.ToView());
 		}
@@ -2149,6 +2620,7 @@ public:
 						~FProtocol0Stage();
 	virtual EStatus		OnData(FStreamReader& Reader, const FMachineContext& Context) override;
 	virtual void		EnterStage(const FMachineContext& Context) override;
+	virtual void		ExitStage(const FMachineContext& Context) override;
 
 private:
 	FTypeRegistry		TypeRegistry;
@@ -2173,6 +2645,20 @@ void FProtocol0Stage::EnterStage(const FMachineContext& Context)
 	Context.Bridge.SetActiveThread(0);
 }
 
+////////////////////////////////////////////////////////////////////////////////
+void FProtocol0Stage::ExitStage(const FMachineContext& Context)
+{
+	// Ensure the transport does not have pending buffers (i.e. event data not yet processed).
+	if (!Transport->IsEmpty())
+	{
+		Context.EmitMessage(EAnalysisMessageSeverity::Warning, TEXT("Transport buffers are not empty at end of analysis (protocol 0)!"));
+	}
+
+#if UE_TRACE_ANALYSIS_DEBUG
+	Transport->DebugEnd();
+#endif
+}
+
 ////////////////////////////////////////////////////////////////////////////////
 FProtocol0Stage::EStatus FProtocol0Stage::OnData(
 	FStreamReader& Reader,
@@ -2202,6 +2688,9 @@ FProtocol0Stage::EStatus FProtocol0Stage::OnData(
 			// There is no need to check size here as the runtime never builds
 			// packets that fragment new-event events.
 			const FTypeRegistry::FTypeInfo* TypeInfo = TypeRegistry.Add(Header->EventData, 0);
+#if UE_TRACE_ANALYSIS_DEBUG
+			Context.Bridge.DebugLogNewEvent(Uid, TypeInfo, BlockSize);
+#endif
 			Context.Bridge.OnNewType(TypeInfo);
 			Transport->Advance(BlockSize);
 			continue;
@@ -2210,9 +2699,14 @@ FProtocol0Stage::EStatus FProtocol0Stage::OnData(
 		const FTypeRegistry::FTypeInfo* TypeInfo = TypeRegistry.Get(Uid);
 		if (TypeInfo == nullptr)
 		{
+			UE_TRACE_ANALYSIS_DEBUG_LOG("Error: Invalid TypeInfo for Uid %u (Tid=%u)", Uid, Context.Bridge.GetActiveThreadId());
 			return EStatus::Error;
 		}
 
+#if UE_TRACE_ANALYSIS_DEBUG
+		Context.Bridge.DebugLogEvent(TypeInfo, BlockSize, 0, Context.Bridge.GetSerial().Value); //???
+#endif
+
 		FEventDataInfo EventDataInfo = {
 			Header->EventData,
 			*TypeInfo,
@@ -2237,18 +2731,19 @@ class FProtocol2Stage
 	: public FAnalysisMachine::FStage
 {
 public:
-								FProtocol2Stage(uint32 Version, FTransport* InTransport);
-								~FProtocol2Stage();
-	virtual EStatus				OnData(FStreamReader& Reader, const FMachineContext& Context) override;
-	virtual void				EnterStage(const FMachineContext& Context) override;
+						FProtocol2Stage(uint32 Version, FTransport* InTransport);
+						~FProtocol2Stage();
+	virtual EStatus		OnData(FStreamReader& Reader, const FMachineContext& Context) override;
+	virtual void		EnterStage(const FMachineContext& Context) override;
+	virtual void		ExitStage(const FMachineContext& Context) override;
 
 protected:
-	virtual int32				OnData(FStreamReader& Reader, FAnalysisBridge& Bridge);
-	int32						OnDataAux(FStreamReader& Reader, FAuxDataCollector& Collector);
-	FTypeRegistry				TypeRegistry;
-	FTransport*					Transport;
-	uint32						ProtocolVersion;
-	uint32						SerialInertia = ~0u;
+	virtual int32		OnData(FStreamReader& Reader, FAnalysisBridge& Bridge);
+	int32				OnDataAux(FStreamReader& Reader, FAuxDataCollector& Collector);
+	FTypeRegistry		TypeRegistry;
+	FTransport*			Transport;
+	uint32				ProtocolVersion;
+	uint32				SerialInertia = ~0u;
 };
 
 ////////////////////////////////////////////////////////////////////////////////
@@ -2303,9 +2798,25 @@ void FProtocol2Stage::EnterStage(const FMachineContext& Context)
 	{
 		Serial.Mask = 0x00ffffff;
 	}
+	UE_TRACE_ANALYSIS_DEBUG_LOG("Serial.Mask = 0x%X", Serial.Mask);
 }
 
 ////////////////////////////////////////////////////////////////////////////////
+void FProtocol2Stage::ExitStage(const FMachineContext& Context)
+{
+	// Ensure the transport does not have pending buffers (i.e. event data not yet processed).
+	if (!Transport->IsEmpty())
+	{
+		Context.EmitMessagef(EAnalysisMessageSeverity::Warning, TEXT("Transport buffers are not empty at end of analysis (protocol %u)!"), ProtocolVersion);
+	}
+
+#if UE_TRACE_ANALYSIS_DEBUG
+	Transport->DebugEnd();
+#endif
+}
+
+////////////////////////////////////////////////////////////////////////////////
+
 FProtocol2Stage::EStatus FProtocol2Stage::OnData(
 	FStreamReader& Reader,
 	const FMachineContext& Context)
@@ -2315,6 +2826,7 @@ FProtocol2Stage::EStatus FProtocol2Stage::OnData(
 	const FTidPacketTransport::ETransportResult Result = InnerTransport->Update();
 	if (Result == FTidPacketTransport::ETransportResult::Error)
 	{
+		UE_TRACE_ANALYSIS_DEBUG_LOG("Error: An error was detected in the transport layer, most likely due to a corrupt trace file.");
 		Context.EmitMessage(
 			EAnalysisMessageSeverity::Error,
 			TEXT("An error was detected in the transport layer, most likely due to a corrupt trace file. See log for details.")
@@ -2345,6 +2857,7 @@ FProtocol2Stage::EStatus FProtocol2Stage::OnData(
 	}
 
 	auto& Serial = Context.Bridge.GetSerial();
+
 	while (true)
 	{
 		uint32 ActiveCount = uint32(Rota.Num());
@@ -2431,6 +2944,7 @@ FProtocol2Stage::EStatus FProtocol2Stage::OnData(
 		if (int32(Serial.Value) < int32(0xc0000000))
 		{
 			Serial.Value = (MinLogSerial & Serial.Mask);
+			UE_TRACE_ANALYSIS_DEBUG_LOG("New Serial.Value: %u", Serial.Value);
 			continue;
 		}
 
@@ -2444,7 +2958,7 @@ FProtocol2Stage::EStatus FProtocol2Stage::OnData(
 
 	// Patch the serial value to try and recover from gaps. Note that the bits
 	// used to wait for the serial-sync event are ignored as that event may never
-	// be reached if leading events are synchronised. Some inertia is added as
+	// be reached if leading events are synchronized. Some inertia is added as
 	// the missing range of events can be in subsequent packets. (Maximum inertia
 	// need only be the size of the tail / io-read-size).
 	if (Rota.Num() > 0)
@@ -2471,6 +2985,7 @@ FProtocol2Stage::EStatus FProtocol2Stage::OnData(
 					{
 						SerialInertia = ~0u;
 						Serial.Value = LowestSerial;
+						UE_TRACE_ANALYSIS_DEBUG_LOG("New Serial.Value: %u (LowestSerial)", Serial.Value);
 						return OnData(Reader, Context);
 					}
 					SerialInertia = ~0u;
@@ -2503,6 +3018,7 @@ int32 FProtocol2Stage::OnData(FStreamReader& Reader, FAnalysisBridge& Bridge)
 		if (TypeInfo == nullptr)
 		{
 			// Event-types may not to be discovered in Uid order.
+			UE_TRACE_ANALYSIS_DEBUG_LOG("Error: Invalid TypeInfo for Uid %u (Tid=%u)", Uid, Bridge.GetActiveThreadId());
 			break;
 		}
 
@@ -2518,6 +3034,9 @@ int32 FProtocol2Stage::OnData(FStreamReader& Reader, FAnalysisBridge& Bridge)
 					const auto* HeaderV1 = (Protocol1::FEventHeader*)Header;
 					if (HeaderV1->Serial != (Serial.Value & Serial.Mask))
 					{
+#if UE_TRACE_ANALYSIS_DEBUG && UE_TRACE_ANALYSIS_DEBUG_LEVEL >= 2
+						UE_TRACE_ANALYSIS_DEBUG_LOG("Tid=%u --> EventSerial=%u", Bridge.GetActiveThreadId(), HeaderV1->Serial);
+#endif
 						return HeaderV1->Serial;
 					}
 					BlockSize += sizeof(*HeaderV1);
@@ -2531,6 +3050,9 @@ int32 FProtocol2Stage::OnData(FStreamReader& Reader, FAnalysisBridge& Bridge)
 					uint32 EventSerial = HeaderSync->SerialLow|(uint32(HeaderSync->SerialHigh) << 16);
 					if (EventSerial != (Serial.Value & Serial.Mask))
 					{
+#if UE_TRACE_ANALYSIS_DEBUG && UE_TRACE_ANALYSIS_DEBUG_LEVEL >= 2
+						UE_TRACE_ANALYSIS_DEBUG_LOG("Tid=%u --> EventSerial=%u", Bridge.GetActiveThreadId(), EventSerial);
+#endif
 						return EventSerial;
 					}
 					BlockSize += sizeof(*HeaderSync);
@@ -2550,15 +3072,26 @@ int32 FProtocol2Stage::OnData(FStreamReader& Reader, FAnalysisBridge& Bridge)
 
 		Reader.Advance(BlockSize);
 
+#if UE_TRACE_ANALYSIS_DEBUG
+		uint32 AuxSize = 0;
+#endif
 		FAuxDataCollector AuxCollector;
 		if (TypeInfo->Flags & FTypeRegistry::FTypeInfo::Flag_MaybeHasAux)
 		{
+#if UE_TRACE_ANALYSIS_DEBUG
+			auto AuxMark = Reader.SaveMark();
+#endif
+
 			int AuxStatus = OnDataAux(Reader, AuxCollector);
 			if (AuxStatus == 0)
 			{
 				Reader.RestoreMark(Mark);
 				break;
 			}
+
+#if UE_TRACE_ANALYSIS_DEBUG
+			AuxSize = uint32(UPTRINT(Reader.SaveMark()) - UPTRINT(AuxMark));
+#endif
 		}
 
 		if ((TypeInfo->Flags & FTypeRegistry::FTypeInfo::Flag_NoSync) == 0)
@@ -2573,10 +3106,17 @@ int32 FProtocol2Stage::OnData(FStreamReader& Reader, FAnalysisBridge& Bridge)
 			// There is no need to check size here as the runtime never builds
 			// packets that fragment new-event events.
 			TypeInfo = TypeRegistry.Add(EventData, 0);
+#if UE_TRACE_ANALYSIS_DEBUG
+			Bridge.DebugLogNewEvent(Uid, TypeInfo, uint32(UPTRINT(Reader.SaveMark()) - UPTRINT(Mark)));
+#endif
 			Bridge.OnNewType(TypeInfo);
 		}
 		else
 		{
+#if UE_TRACE_ANALYSIS_DEBUG
+			const uint32 EventSize = uint32(UPTRINT(Reader.SaveMark()) - UPTRINT(Mark));
+			Bridge.DebugLogEvent(TypeInfo, EventSize - AuxSize, AuxSize, Serial.Value - 1);
+#endif
 			FEventDataInfo EventDataInfo = {
 				EventData,
 				*TypeInfo,
@@ -2644,12 +3184,12 @@ class FProtocol4Stage
 	: public FProtocol2Stage
 {
 public:
-					FProtocol4Stage(uint32 Version, FTransport* InTransport);
+						FProtocol4Stage(uint32 Version, FTransport* InTransport);
 
 private:
-	virtual int32	OnData(FStreamReader& Reader, FAnalysisBridge& Bridge) override;
-	int32			OnDataImpl(FStreamReader& Reader, FAnalysisBridge& Bridge);
-	int32			OnDataKnown(uint32 Uid, FStreamReader& Reader, FAnalysisBridge& Bridge);
+	virtual int32		OnData(FStreamReader& Reader, FAnalysisBridge& Bridge) override;
+	int32				OnDataImpl(FStreamReader& Reader, FAnalysisBridge& Bridge);
+	int32				OnDataKnown(uint32 Uid, FStreamReader& Reader, FAnalysisBridge& Bridge);
 };
 
 ////////////////////////////////////////////////////////////////////////////////
@@ -2703,6 +3243,10 @@ int32 FProtocol4Stage::OnDataImpl(FStreamReader& Reader, FAnalysisBridge& Bridge
 	if (Uid < Bridge.GetUserUidBias())
 	{
 		Reader.Advance(UidBytes);
+#if UE_TRACE_ANALYSIS_DEBUG
+		FAnalysisState& State = Bridge.GetState();
+		State.TempUidBytes = UidBytes;
+#endif
 		if (!OnDataKnown(Uid, Reader, Bridge))
 		{
 			Reader.RestoreMark(Mark);
@@ -2715,6 +3259,13 @@ int32 FProtocol4Stage::OnDataImpl(FStreamReader& Reader, FAnalysisBridge& Bridge
 	const FTypeRegistry::FTypeInfo* TypeInfo = TypeRegistry.Get(Uid);
 	if (TypeInfo == nullptr)
 	{
+		UE_TRACE_ANALYSIS_DEBUG_LOG("Error: Invalid TypeInfo for Uid %u (Tid=%u)", Uid, Bridge.GetActiveThreadId());
+		return 1;
+	}
+
+	if (!ensure(UidBytes == 2)) // see Protocol4::FEventHeader
+	{
+		UE_TRACE_ANALYSIS_DEBUG_LOG("Error: Invalid trace stream: Tid=%u Uid=%u encoded on 1 byte (instead of 2 bytes).", Bridge.GetActiveThreadId(), Uid);
 		return 1;
 	}
 
@@ -2739,6 +3290,9 @@ int32 FProtocol4Stage::OnDataImpl(FStreamReader& Reader, FAnalysisBridge& Bridge
 		uint32 EventSerial = HeaderSync->SerialLow|(uint32(HeaderSync->SerialHigh) << 16);
 		if (EventSerial != (Serial.Value & Serial.Mask))
 		{
+#if UE_TRACE_ANALYSIS_DEBUG && UE_TRACE_ANALYSIS_DEBUG_LEVEL >= 2
+			UE_TRACE_ANALYSIS_DEBUG_LOG("Tid=%u --> EventSerial=%u", Bridge.GetActiveThreadId(), EventSerial);
+#endif
 			return ~EventSerial;
 		}
 		BlockSize += sizeof(*HeaderSync);
@@ -2756,6 +3310,10 @@ int32 FProtocol4Stage::OnDataImpl(FStreamReader& Reader, FAnalysisBridge& Bridge
 
 	Reader.Advance(BlockSize);
 
+#if UE_TRACE_ANALYSIS_DEBUG
+	uint32 AuxSize = 0;
+#endif
+
 	// Collect auxiliary data
 	FAuxDataCollector AuxCollector;
 	if (TypeInfo->Flags & FTypeRegistry::FTypeInfo::Flag_MaybeHasAux)
@@ -2768,6 +3326,10 @@ int32 FProtocol4Stage::OnDataImpl(FStreamReader& Reader, FAnalysisBridge& Bridge
 			Reader.Advance(sizeof(Protocol4::FEventHeader) + TypeInfo->EventSize);
 		}
 
+#if UE_TRACE_ANALYSIS_DEBUG
+		auto AuxMark = Reader.SaveMark();
+#endif
+
 		int AuxStatus = OnDataAux(Reader, AuxCollector);
 		if (AuxStatus == 0)
 		{
@@ -2775,6 +3337,10 @@ int32 FProtocol4Stage::OnDataImpl(FStreamReader& Reader, FAnalysisBridge& Bridge
 			return 1;
 		}
 
+#if UE_TRACE_ANALYSIS_DEBUG
+		AuxSize = uint32(UPTRINT(Reader.SaveMark()) - UPTRINT(AuxMark));
+#endif
+
 		// User error could have resulted in less space being used that was
 		// allocated for important events. So we can't assume that aux data
 		// reading has read all the way up to the next event. So we use marks
@@ -2784,6 +3350,11 @@ int32 FProtocol4Stage::OnDataImpl(FStreamReader& Reader, FAnalysisBridge& Bridge
 		}
 	}
 
+#if UE_TRACE_ANALYSIS_DEBUG
+	const uint32 EventSize = uint32(UPTRINT(Reader.SaveMark()) - UPTRINT(Mark));
+	Bridge.DebugLogEvent(TypeInfo, EventSize - AuxSize, AuxSize, Serial.Value);
+#endif
+
 	// Maintain sync
 	if ((TypeInfo->Flags & FTypeRegistry::FTypeInfo::Flag_NoSync) == 0)
 	{
@@ -2810,22 +3381,38 @@ int32 FProtocol4Stage::OnDataKnown(
 	FStreamReader& Reader,
 	FAnalysisBridge& Bridge)
 {
+#if UE_TRACE_ANALYSIS_DEBUG
+	FAnalysisState& State = Bridge.GetState();
+#endif
+
 	switch (Uid)
 	{
 	case Protocol4::EKnownEventUids::NewEvent:
 		{
 			const auto* Size = Reader.GetPointer<uint16>();
+			check(Size != nullptr);
+			const void* EventTypeAndData = Reader.GetPointer(sizeof(*Size) + *Size);
+			check(EventTypeAndData != nullptr);
 			const FTypeRegistry::FTypeInfo* TypeInfo = TypeRegistry.Add(Size + 1, 4);
+#if UE_TRACE_ANALYSIS_DEBUG
+			Bridge.DebugLogNewEvent(Uid, TypeInfo, uint32(State.TempUidBytes + sizeof(*Size) + *Size));
+#endif
 			Bridge.OnNewType(TypeInfo);
 			Reader.Advance(sizeof(*Size) + *Size);
 		}
 		return 1;
 
 	case Protocol4::EKnownEventUids::EnterScope:
+#if UE_TRACE_ANALYSIS_DEBUG
+		Bridge.DebugLogEnterScopeEvent(Uid, State.TempUidBytes);
+#endif
 		Bridge.EnterScope();
 		return 1;
 
 	case Protocol4::EKnownEventUids::LeaveScope:
+#if UE_TRACE_ANALYSIS_DEBUG
+		Bridge.DebugLogLeaveScopeEvent(Uid, State.TempUidBytes);
+#endif
 		Bridge.LeaveScope();
 		return 1;
 
@@ -2838,6 +3425,9 @@ int32 FProtocol4Stage::OnDataKnown(
 			}
 
 			const uint64 RelativeTimestamp = *(uint64*)(Stamp - 1) >> 8;
+#if UE_TRACE_ANALYSIS_DEBUG
+			Bridge.DebugLogEnterScopeEvent(Uid, RelativeTimestamp, uint32(State.TempUidBytes + sizeof(uint64) - 1));
+#endif
 			Bridge.EnterScope(RelativeTimestamp);
 
 			Reader.Advance(sizeof(uint64) - 1);
@@ -2853,6 +3443,9 @@ int32 FProtocol4Stage::OnDataKnown(
 			}
 
 			const uint64 RelativeTimestamp = *(uint64*)(Stamp - 1) >> 8;
+#if UE_TRACE_ANALYSIS_DEBUG
+			Bridge.DebugLogLeaveScopeEvent(Uid, RelativeTimestamp, uint32(State.TempUidBytes + sizeof(uint64) - 1));
+#endif
 			Bridge.LeaveScope(RelativeTimestamp);
 
 			Reader.Advance(sizeof(uint64) - 1);
@@ -2860,6 +3453,7 @@ int32 FProtocol4Stage::OnDataKnown(
 		return 1;
 
 	default:
+		UE_TRACE_ANALYSIS_DEBUG_LOG("Error: Cannot process known event %llu with Uid %u", State.TotalEventCount, Uid);
 		return 0;
 	};
 }
@@ -2873,7 +3467,8 @@ class FProtocol5Stage
 	: public FAnalysisMachine::FStage
 {
 public:
-							FProtocol5Stage(FTransport* InTransport);
+						FProtocol5Stage(FTransport* InTransport);
+	virtual void		ExitStage(const FMachineContext& Context) override;
 
 protected:
 	struct alignas(16) FEventDesc
@@ -2893,10 +3488,19 @@ protected:
 		union
 		{
 			uint32			GapLength;
-			const uint8*	Data;
+			const uint8*	Data = nullptr;
 		};
+#if UE_TRACE_ANALYSIS_DEBUG
+		uint32				EventSize = 0;
+		uint32				AuxSize = 0;
+		uint64				Reserved = 0;
+#endif
 	};
+#if UE_TRACE_ANALYSIS_DEBUG
+	static_assert(sizeof(FEventDesc) == 32, "");
+#else
 	static_assert(sizeof(FEventDesc) == 16, "");
+#endif
 
 	struct alignas(16) FEventDescStream
 	{
@@ -2942,7 +3546,6 @@ protected:
 	virtual EStatus			OnData(FStreamReader& Reader, const FMachineContext& Context) override;
 	EStatus					OnDataNewEvents(const FMachineContext& Context);
 	EStatus					OnDataImportant(const FMachineContext& Context);
-	EStatus					OnDataNonCachedImportant(const FMachineContext& Context);
 	EStatus					OnDataNormal(const FMachineContext& Context);
 	int32					ParseImportantEvents(FStreamReader& Reader, EventDescArray& OutEventDescs, const FMachineContext& Context);
 	int32					ParseEvents(FStreamReader& Reader, EventDescArray& OutEventDescs, const FMachineContext& Context);
@@ -2956,15 +3559,21 @@ protected:
 	void					DetectSerialGaps(TArray<FEventDescStream>& EventDescHeap);
 	template <typename Callback>
 	void					ForEachSerialGap(const TArray<FEventDescStream>& EventDescHeap, Callback&& InCallback);
+#if UE_TRACE_ANALYSIS_DEBUG
+	void					PrintParsedEvent(int EventIndex, const FEventDesc& EventDesc, int32 Size);
+#endif // UE_TRACE_ANALYSIS_DEBUG
 
 	FTypeRegistry			TypeRegistry;
 	FTidPacketTransport&	Transport;
 	EventDescArray			EventDescs;
 	EventDescArray			SerialGaps;
 	uint32					NextSerial = ~0u;
+	uint32					OldNextSerial = ~0u;
+	uint32					NextSerialWaitCount = 0;
 	uint32					SyncCount;
 	uint32					EventVersion = 4; //Protocol version 5 uses the event version from protocol 4
 	bool					bSkipSerialError = false;
+	bool					bSkipSerial = false;
 };
 
 ////////////////////////////////////////////////////////////////////////////////
@@ -2975,6 +3584,21 @@ FProtocol5Stage::FProtocol5Stage(FTransport* InTransport)
 	EventDescs.Reserve(8 << 10);
 }
 
+////////////////////////////////////////////////////////////////////////////////
+void FProtocol5Stage::ExitStage(const FMachineContext& Context)
+{
+	// Ensure the transport does not have pending buffers (i.e. event data not yet processed).
+	if (!Transport.IsEmpty())
+	{
+		Context.EmitMessage(EAnalysisMessageSeverity::Warning, TEXT("Transport buffers are not empty at end of analysis (protocol 5)!"));
+	}
+
+#if UE_TRACE_ANALYSIS_DEBUG
+	Context.Bridge.GetSerial().Value = NextSerial;
+	Transport.DebugEnd();
+#endif
+}
+
 ////////////////////////////////////////////////////////////////////////////////
 FProtocol5Stage::EStatus FProtocol5Stage::OnData(
 	FStreamReader& Reader,
@@ -2991,39 +3615,43 @@ FProtocol5Stage::EStatus FProtocol5Stage::OnData(
 		return EStatus::Error;
 	}
 
-	// New-events. They must be processed before anything else otherwise events
-	// can not be interpreted.
-	EStatus Ret = OnDataNewEvents(Context);
-	if (Ret == EStatus::Error)
+	do
 	{
-		return Ret;
-	}
+		// New-events. They must be processed before anything else otherwise events
+		// can not be interpreted.
+		EStatus Ret = OnDataNewEvents(Context);
+		if (Ret == EStatus::Error)
+		{
+			return Ret;
+		}
 
-	// Important events
-	Ret = OnDataImportant(Context);
-	if (Ret == EStatus::Error)
-	{
-		return Ret;
-	}
-	bool bNotEnoughData = (Ret == EStatus::NotEnoughData);
+		// Important events
+		Ret = OnDataImportant(Context);
+		if (Ret == EStatus::Error)
+		{
+			return Ret;
+		}
+		bool bNotEnoughData = (Ret == EStatus::NotEnoughData);
 
-	// Normal events
-	Ret = OnDataNormal(Context);
-	if (Ret == EStatus::Error)
-	{
-		return Ret;
-	}
-	if (Ret == EStatus::Sync)
-	{
-		// After processing a SYNC packet, we need to read data once more.
-		return OnData(Reader, Context);
-	}
-	bNotEnoughData |= (Ret == EStatus::NotEnoughData);
+		// Normal events
+		Ret = OnDataNormal(Context);
+		if (Ret == EStatus::Error)
+		{
+			return Ret;
+		}
+		if (Ret == EStatus::Sync)
+		{
+			// After processing a SYNC packet, we need to read data once more.
+			return OnData(Reader, Context);
+		}
+		bNotEnoughData |= (Ret == EStatus::NotEnoughData);
 
-	if (bNotEnoughData)
-	{
-		return EStatus::NotEnoughData;
+		if (bNotEnoughData && !bSkipSerial)
+		{
+			return EStatus::NotEnoughData;
+		}
 	}
+	while (bSkipSerial);
 
 	return Reader.CanMeetDemand() ? EStatus::Continue : EStatus::EndOfStream;
 }
@@ -3041,12 +3669,16 @@ FProtocol5Stage::EStatus FProtocol5Stage::OnDataNewEvents(const FMachineContext&
 
 	if (ParseImportantEvents(*ThreadReader, EventDescs, Context) < 0)
 	{
+		UE_TRACE_ANALYSIS_DEBUG_LOG("Error: Failed to parse important events");
 		return EStatus::Error;
 	}
 
 	for (const FEventDesc& EventDesc : EventDescs)
 	{
 		const FTypeRegistry::FTypeInfo* TypeInfo = TypeRegistry.Add(EventDesc.Data, EventVersion);
+#if UE_TRACE_ANALYSIS_DEBUG
+		Context.Bridge.DebugLogNewEvent(uint32(EventDesc.Uid), TypeInfo, EventDesc.EventSize + EventDesc.AuxSize);
+#endif
 		Context.Bridge.OnNewType(TypeInfo);
 	}
 
@@ -3068,6 +3700,7 @@ FProtocol5Stage::EStatus FProtocol5Stage::OnDataImportant(const FMachineContext&
 
 	if (ParseImportantEvents(*ThreadReader, EventDescs, Context) < 0)
 	{
+		UE_TRACE_ANALYSIS_DEBUG_LOG("Error: Failed to parse important events");
 		return EStatus::Error;
 	}
 
@@ -3086,6 +3719,7 @@ FProtocol5Stage::EStatus FProtocol5Stage::OnDataImportant(const FMachineContext&
 	Context.Bridge.SetActiveThread(ETransportTid::Importants);
 	if (DispatchEvents(Context, EventDescs.GetData(), EventDescs.Num() - 1) < 0)
 	{
+		UE_TRACE_ANALYSIS_DEBUG_LOG("Error: Failed to dispatch important events");
 		return EStatus::Error;
 	}
 
@@ -3099,14 +3733,14 @@ int32 FProtocol5Stage::ParseImportantEvents(FStreamReader& Reader, EventDescArra
 
 	while (true)
 	{
-		int32 Remaining = Reader.GetRemaining();
+		uint32 Remaining = Reader.GetRemaining();
 		if (Remaining < sizeof(FImportantEventHeader))
 		{
 			return 1;
 		}
 
 		const auto* Header = Reader.GetPointerUnchecked<FImportantEventHeader>();
-		if (Remaining < int32(Header->Size) + sizeof(FImportantEventHeader))
+		if (Remaining < uint32(Header->Size) + sizeof(FImportantEventHeader))
 		{
 			return 1;
 		}
@@ -3122,6 +3756,9 @@ int32 FProtocol5Stage::ParseImportantEvents(FStreamReader& Reader, EventDescArra
 		// registry but this way avoid raveling things together.
 		if (Uid == EKnownUids::NewEvent)
 		{
+#if UE_TRACE_ANALYSIS_DEBUG
+			EventDesc.EventSize = sizeof(*Header) + Header->Size;
+#endif
 			OutEventDescs.Add(EventDesc);
 			Reader.Advance(sizeof(*Header) + Header->Size);
 			continue;
@@ -3130,10 +3767,14 @@ int32 FProtocol5Stage::ParseImportantEvents(FStreamReader& Reader, EventDescArra
 		const FTypeRegistry::FTypeInfo* TypeInfo = TypeRegistry.Get(Uid);
 		if (TypeInfo == nullptr)
 		{
-			Context.EmitMessagef(EAnalysisMessageSeverity::Warning, TEXT("An unknown event UID (%d) was encountered."), Uid);
+			UE_TRACE_ANALYSIS_DEBUG_LOG("UID %u (0x%X) was not declared yet.", Uid, Uid);
 			return 1;
 		}
 
+#if UE_TRACE_ANALYSIS_DEBUG
+		EventDesc.EventSize = sizeof(*Header) + TypeInfo->EventSize;
+#endif
+
 		if (TypeInfo->Flags & FTypeRegistry::FTypeInfo::Flag_MaybeHasAux)
 		{
 			EventDesc.bHasAux = 1;
@@ -3160,10 +3801,16 @@ int32 FProtocol5Stage::ParseImportantEvents(FStreamReader& Reader, EventDescArra
 				AuxDesc.Serial = ESerial::Ignored;
 
 				Cursor = AuxHeader->Data + (AuxHeader->Pack >> FAuxHeader::SizeShift);
+
+#if UE_TRACE_ANALYSIS_DEBUG
+				AuxDesc.EventSize = sizeof(FAuxHeader);
+				AuxDesc.AuxSize = (AuxHeader->Pack >> FAuxHeader::SizeShift);
+#endif
 			}
 
 			if (Cursor[0] != uint8(EKnownUids::AuxDataTerminal))
 			{
+				UE_TRACE_ANALYSIS_DEBUG_LOG("Error: Expecting AuxDataTerminal event");
 				Context.EmitMessage(EAnalysisMessageSeverity::Warning, TEXT("Expected an aux data terminal in the stream."));
 				return -1;
 			}
@@ -3186,29 +3833,34 @@ FProtocol5Stage::EStatus FProtocol5Stage::OnDataNormal(const FMachineContext& Co
 	TArray<FEventDescStream> EventDescHeap;
 	EventDescHeap.Reserve(Transport.GetThreadCount());
 
-	bool bSkipSerial = false;
+	bSkipSerial = false;
 
 	for (uint32 i = ETransportTid::Bias, n = Transport.GetThreadCount(); i < n; ++i)
 	{
 		uint32 NumEventDescs = EventDescs.Num();
 
-		TRACE_ANALYSIS_DEBUG("Thread: %03d Id:%04x", i, Transport.GetThreadId(i));
+#if UE_TRACE_ANALYSIS_DEBUG && UE_TRACE_ANALYSIS_DEBUG_LEVEL >= 3
+		UE_TRACE_ANALYSIS_DEBUG_LOG("Thread:%u Id:%u", i, Transport.GetThreadId(i));
+#endif
 
 		// Extract all the events in the stream for this thread
 		FStreamReader* ThreadReader = Transport.GetThreadStream(i);
 
-		// Stop analysis if this thread has accumulated too much data.
+		// Test if analysis has accumulated too much data for this thread.
 		// This can happen on corrupted traces (ex. with serial sync events missing or out of order).
-		constexpr int32 MaxAccumulatedBytes = 512 * 1024 * 1024;
+		constexpr uint32 MaxAccumulatedBytes = 2'000'000'000u;
 		if (ThreadReader->GetRemaining() > MaxAccumulatedBytes)
 		{
+			UE_TRACE_ANALYSIS_DEBUG_LOG("Error: Trace analysis accumulated too much data (%.2f MiB on thread %u) and will start to skip the missing serial sync events!",
+				(double)ThreadReader->GetRemaining() / (1024.0 * 1024.0),
+				Transport.GetThreadId(i));
 			if (!bSkipSerialError)
 			{
 				bSkipSerialError = true;
 				Context.EmitMessagef(
 					EAnalysisMessageSeverity::Error,
-					TEXT("Trace analysis accumulated too much data (%d MiB on thread %u) and will start to skip the missing serial sync events!"),
-					ThreadReader->GetRemaining() / (1024 * 1024),
+					TEXT("Trace analysis accumulated too much data (%.2f MiB on thread %u) and will start to skip the missing serial sync events!"),
+					(double)ThreadReader->GetRemaining() / (1024.0 * 1024.0),
 					Transport.GetThreadId(i)
 				);
 			}
@@ -3217,6 +3869,7 @@ FProtocol5Stage::EStatus FProtocol5Stage::OnDataNormal(const FMachineContext& Co
 
 		if (ParseEvents(*ThreadReader, EventDescs, Context) < 0)
 		{
+			UE_TRACE_ANALYSIS_DEBUG_LOG("Error: Failed to parse events");
 			return EStatus::Error;
 		}
 
@@ -3235,7 +3888,9 @@ FProtocol5Stage::EStatus FProtocol5Stage::OnDataNormal(const FMachineContext& Co
 			EventDescHeap.Add(Out);
 		}
 
-		TRACE_ANALYSIS_DEBUG("Thread: %03d bNotEnoughData:%d", i, bNotEnoughData);
+#if UE_TRACE_ANALYSIS_DEBUG && UE_TRACE_ANALYSIS_DEBUG_LEVEL >= 3
+		UE_TRACE_ANALYSIS_DEBUG_LOG("Thread:%u bNotEnoughData:%d", i, bNotEnoughData);
+#endif
 	}
 
 	// Now EventDescs is stable we can convert the indices into pointers
@@ -3244,37 +3899,73 @@ FProtocol5Stage::EStatus FProtocol5Stage::OnDataNormal(const FMachineContext& Co
 		Stream.EventDescs = EventDescs.GetData() + Stream.ContainerIndex;
 	}
 
+#if UE_TRACE_ANALYSIS_DEBUG
+	FAnalysisState& State = Context.Bridge.GetState();
+	if (EventDescs.Num() > State.MaxEventDescs)
+	{
+		State.MaxEventDescs = EventDescs.Num();
+	}
+#endif
+
 	const bool bSync = (SyncCount != Transport.GetSyncCount());
 
-	if (bSync || !bSkipSerial)
+	int32 NumAvailableEvents = EventDescs.Num();
+
+	// Try to dispatch the parsed events.
 	{
 		int32 NumDispatchedEvents = DispatchNormalEvents(Context, EventDescHeap);
 		if (NumDispatchedEvents < 0)
 		{
 			return EStatus::Error;
 		}
+#if UE_TRACE_ANALYSIS_DEBUG && UE_TRACE_ANALYSIS_DEBUG_LEVEL >= 2
+		UE_TRACE_ANALYSIS_DEBUG_LOG("Dispatched %d normal events (%d --> %d)", NumDispatchedEvents, NumAvailableEvents, NumAvailableEvents - NumDispatchedEvents);
+#endif
+		NumAvailableEvents -= NumDispatchedEvents;
+		check(NumAvailableEvents >= 0);
+
+		// Count how many times we dispatched events, but without dispatching any "sync" event.
+		if (OldNextSerial == NextSerial)
+		{
+			++NextSerialWaitCount;
+		}
+		else
+		{
+			OldNextSerial = NextSerial;
+			NextSerialWaitCount = 0;
+		}
 	}
-	else // bSkipSerial && !bSync
+
+	// Test if analysis has accumulated too much data (parsed events not dispatched yet).
+	// But, only enforce the limit after we have received at least one SYNC package
+	// (e.g. server traces can accumulate large amounts of data before first SYNC event).
+	constexpr int32 MaxAvailableParsedEvents = 10'000'000;
+	constexpr uint32 MaxNextSerialWaitCount = 20;
+	if (SyncCount > 0 &&
+		NumAvailableEvents > MaxAvailableParsedEvents &&
+		NextSerialWaitCount > MaxNextSerialWaitCount)
 	{
-		// Dispatch at least half of the parsed events.
-		int32 TotalNumDispatchedEvents = 0;
-		do
+		UE_TRACE_ANALYSIS_DEBUG_LOG("Error: Trace analysis accumulated too much data (%d parsed events) and will start to skip the missing serial sync events!", NumAvailableEvents);
+		if (!bSkipSerialError)
 		{
-			NextSerial = ~0u; // skip missing serial sync events
-
-			int32 NumDispatchedEvents = DispatchNormalEvents(Context, EventDescHeap);
-			if (NumDispatchedEvents < 0)
-			{
-				return EStatus::Error;
-			}
-
-			TotalNumDispatchedEvents += NumDispatchedEvents;
-			if (TotalNumDispatchedEvents > EventDescs.Num() / 2)
-			{
-				bSkipSerial = false;
-			}
+			bSkipSerialError = true;
+			Context.EmitMessagef(
+				EAnalysisMessageSeverity::Error,
+				TEXT("Trace analysis accumulated too much data (%d parsed events) and will start to skip the missing serial sync events!"), NumAvailableEvents);
 		}
-		while (bSkipSerial);
+		// Skip serials and continue to dispatch parsed events.
+		bSkipSerial = true;
+		NextSerialWaitCount = 0;
+		int32 NumDispatchedEvents = DispatchNormalEvents(Context, EventDescHeap);
+		if (NumDispatchedEvents < 0)
+		{
+			return EStatus::Error;
+		}
+#if UE_TRACE_ANALYSIS_DEBUG && UE_TRACE_ANALYSIS_DEBUG_LEVEL >= 2
+		UE_TRACE_ANALYSIS_DEBUG_LOG("Skipped serials and dispatched %d normal events (%d --> %d)", NumDispatchedEvents, NumAvailableEvents, NumAvailableEvents - NumDispatchedEvents);
+#endif
+		NumAvailableEvents -= NumDispatchedEvents;
+		check(NumAvailableEvents >= 0);
 	}
 
 	// If there are any streams left in the heap then we are unable to proceed
@@ -3303,12 +3994,16 @@ FProtocol5Stage::EStatus FProtocol5Stage::OnDataNormal(const FMachineContext& Co
 ////////////////////////////////////////////////////////////////////////////////
 int32 FProtocol5Stage::DispatchNormalEvents(const FMachineContext& Context, TArray<FEventDescStream>& EventDescHeap)
 {
+#if UE_TRACE_ANALYSIS_DEBUG && UE_TRACE_ANALYSIS_DEBUG_LEVEL >= 2
+	UE_TRACE_ANALYSIS_DEBUG_LOG("Queued event descs: %d", EventDescs.Num());
+#endif
+
 	int32 NumDispatchedUnsyncEvents = 0;
 
-	// Process leading unsynchronised events so that each stream starts with a sychronised event.
+	// Process leading unsynchronized events so that each stream starts with a synchronized event.
 	for (FEventDescStream& Stream : EventDescHeap)
 	{
-		// Extract a run of consecutive unsynchronised events
+		// Extract a run of consecutive unsynchronized events
 		const FEventDesc* EndDesc = Stream.EventDescs;
 		for (; EndDesc->Serial == ESerial::Ignored; ++EndDesc);
 
@@ -3323,6 +4018,7 @@ int32 FProtocol5Stage::DispatchNormalEvents(const FMachineContext& Context, TArr
 
 			if (DispatchEvents(Context, StartDesc, DescNum) < 0)
 			{
+				UE_TRACE_ANALYSIS_DEBUG_LOG("Error: Failed to dispatch events");
 				return -1;
 			}
 
@@ -3345,10 +4041,32 @@ int32 FProtocol5Stage::DispatchNormalEvents(const FMachineContext& Context, TArr
 	// A min-heap is used to peel off groups of events by lowest serial
 	EventDescHeap.Heapify(FSerialDistancePredicate{NextSerial});
 
+#if UE_TRACE_ANALYSIS_DEBUG && UE_TRACE_ANALYSIS_DEBUG_LEVEL >= 2
+	{
+		const FEventDescStream& TopStream = EventDescHeap.HeapTop();
+		UE_TRACE_ANALYSIS_DEBUG_LOG("NextSerial=%u LowestSerial=%d (Tid=%u)", NextSerial, TopStream.EventDescs[0].Serial, TopStream.ThreadId);
+	}
+#endif
+
 	// Events must be consumed contiguously.
-	if (NextSerial == ~0u && Transport.GetSyncCount())
+	if (bSkipSerial)
 	{
-		NextSerial = EventDescHeap.HeapTop().EventDescs[0].Serial;
+		uint32 LowestSerial = EventDescHeap.HeapTop().EventDescs[0].Serial;
+		if (LowestSerial != NextSerial)
+		{
+			UE_TRACE_ANALYSIS_DEBUG_LOG("Warning: NextSerial skips %lld events (from %u to %u)", (int64)LowestSerial - (int64)NextSerial, NextSerial, LowestSerial);
+			NextSerial = LowestSerial;
+		}
+	}
+	else
+	if (NextSerial == ~0u)
+	{
+		uint32 LowestSerial = EventDescHeap.HeapTop().EventDescs[0].Serial;
+		if (Transport.GetSyncCount() || LowestSerial == 0)
+		{
+			NextSerial = LowestSerial;
+			UE_TRACE_ANALYSIS_DEBUG_LOG("NextSerial=%u", NextSerial);
+		}
 	}
 
 	DetectSerialGaps(EventDescHeap);
@@ -3356,6 +4074,7 @@ int32 FProtocol5Stage::DispatchNormalEvents(const FMachineContext& Context, TArr
 	int32 NumDispatchedEvents = DispatchEvents(Context, EventDescHeap);
 	if (NumDispatchedEvents < 0)
 	{
+		UE_TRACE_ANALYSIS_DEBUG_LOG("Error: Failed to dispatch events");
 		return -1;
 	}
 
@@ -3367,6 +4086,10 @@ int32 FProtocol5Stage::DispatchEvents(
 	const FMachineContext& Context,
 	TArray<FEventDescStream>& EventDescHeap)
 {
+#if UE_TRACE_ANALYSIS_DEBUG
+	FAnalysisState& State = Context.Bridge.GetState();
+#endif
+
 	auto UpdateHeap = [&] (const FEventDescStream& Stream, const FEventDesc* EventDesc)
 	{
 		if (EventDesc->Serial != ESerial::Terminal)
@@ -3394,13 +4117,25 @@ int32 FProtocol5Stage::DispatchEvents(
 		{
 			NextSerial = EndDesc->Serial + EndDesc->GapLength;
 			NextSerial &= ESerial::Mask;
+#if UE_TRACE_ANALYSIS_DEBUG
+			if (NextSerial != EndDesc->Serial + EndDesc->GapLength)
+			{
+				State.SerialWrappedCount++;
+			}
+			State.NumSkippedSerialGaps++;
+			UE_TRACE_ANALYSIS_DEBUG_LOG("Skip serial gap (%u +%u) --> NextSerial=%u", EndDesc->Serial, EndDesc->GapLength, NextSerial);
+#endif
 			UpdateHeap(Stream, EndDesc + 1);
 			continue;
 		}
 
-		// Extract a run of consecutive events (plus runs of unsynchronised ones)
+		// Extract a run of consecutive events (plus runs of unsynchronized ones).
 		if (EndDesc->Serial == NextSerial)
 		{
+#if UE_TRACE_ANALYSIS_DEBUG
+			const uint32 CurrentSerial = NextSerial;
+#endif
+
 			do
 			{
 				NextSerial = (NextSerial + 1) & ESerial::Mask;
@@ -3412,11 +4147,71 @@ int32 FProtocol5Stage::DispatchEvents(
 				while (EndDesc->Serial == ESerial::Ignored);
 			}
 			while (EndDesc->Serial == NextSerial);
+
+#if UE_TRACE_ANALYSIS_DEBUG
+			if (NextSerial < CurrentSerial)
+			{
+				++State.SerialWrappedCount;
+			}
+#endif
 		}
 		else
 		{
-			// The lowest known serial number is not low enough so we are unable
-			// to proceed any further.
+#if UE_TRACE_ANALYSIS_DEBUG
+			if (uint32(EndDesc->Serial) < NextSerial &&
+				NextSerial != ~0u &&
+				NextSerial - uint32(EndDesc->Serial) < ESerial::Range/2)
+			{
+				UE_TRACE_ANALYSIS_DEBUG_LOG("Error: Lowest serial %d (Tid=%u Uid=%u) is too low (NextSerial=%u; %d event descs) !!!", EndDesc->Serial, Stream.ThreadId, uint32(EndDesc->Uid), NextSerial, EventDescs.Num());
+			}
+			else
+			{
+#if UE_TRACE_ANALYSIS_DEBUG_LEVEL >= 2
+				UE_TRACE_ANALYSIS_DEBUG_LOG("Lowest serial %d (Tid=%u Uid=%u) is not low enough (NextSerial=%u; %d event descs)", EndDesc->Serial, Stream.ThreadId, uint32(EndDesc->Uid), NextSerial, EventDescs.Num());
+#endif
+			}
+#if UE_TRACE_ANALYSIS_DEBUG_LEVEL >= 2
+			UE_TRACE_ANALYSIS_DEBUG_LOG("Available streams (SerialWrappedCount=%u):", State.SerialWrappedCount);
+			for (const FEventDescStream& EventDescStream : EventDescHeap)
+			{
+				uint32 BufferSize = 0;
+				uint32 DataSize = 0;
+				FTidPacketTransport* InnerTransport = (FTidPacketTransport*)(&Transport);
+				for (uint32 i = 0, n = InnerTransport->GetThreadCount(); i < n; ++i)
+				{
+					FStreamBuffer* ThreadReader = (FStreamBuffer*)InnerTransport->GetThreadStream(i);
+					uint32 ThreadId = InnerTransport->GetThreadId(i);
+					if (ThreadId == EventDescStream.ThreadId)
+					{
+						BufferSize = ThreadReader->GetBufferSize();
+						DataSize = ThreadReader->GetRemaining();
+					}
+				}
+				if (EventDescStream.EventDescs->Serial == NextSerial)
+				{
+					UE_TRACE_ANALYSIS_DEBUG_LOG("  Tid=%u : Serial=%u BufferSize=%u DataSize=%u (next)", EventDescStream.ThreadId, EventDescStream.EventDescs->Serial, BufferSize, DataSize);
+				}
+				else
+				{
+					UE_TRACE_ANALYSIS_DEBUG_LOG("  Tid=%u : Serial=%u BufferSize=%u DataSize=%u", EventDescStream.ThreadId, EventDescStream.EventDescs->Serial, BufferSize, DataSize);
+				}
+			}
+#endif // UE_TRACE_ANALYSIS_DEBUG_LEVEL
+#endif // UE_TRACE_ANALYSIS_DEBUG
+
+#if 0
+			int32 MinSerial = EndDesc->Serial;
+			EventDescHeap.Heapify();
+			const FEventDescStream& MinStream = EventDescHeap.HeapTop();
+			const FEventDesc* MinDesc = MinStream.EventDescs;
+			if (MinDesc->Serial < MinSerial)
+			{
+				UE_TRACE_ANALYSIS_DEBUG_LOG("Try one more time with lowest serial %d (Tid=%u Uid=%u)", MinDesc->Serial, MinStream.ThreadId, uint32(MinDesc->Uid));
+				continue;
+			}
+#endif
+
+			// The lowest known serial number is not low enough so we are unable to proceed any further.
 			break;
 		}
 
@@ -3427,6 +4222,7 @@ int32 FProtocol5Stage::DispatchEvents(
 		NumDispatchedEvents += DescNum;
 		if (DispatchEvents(Context, StartDesc, DescNum) < 0)
 		{
+			UE_TRACE_ANALYSIS_DEBUG_LOG("Error: Failed to dispatch events");
 			return -1;
 		}
 
@@ -3437,19 +4233,54 @@ int32 FProtocol5Stage::DispatchEvents(
 	return NumDispatchedEvents;
 }
 
+////////////////////////////////////////////////////////////////////////////////
+#if UE_TRACE_ANALYSIS_DEBUG
+void FProtocol5Stage::PrintParsedEvent(int EventIndex, const FEventDesc& EventDesc, int32 Size)
+{
+	UE_TRACE_ANALYSIS_DEBUG_BeginStringBuilder();
+	UE_TRACE_ANALYSIS_DEBUG_Appendf("Event=%-6d Uid=%-4u ", EventIndex, uint32(EventDesc.Uid));
+	if (EventDesc.Serial >= ESerial::Range)
+	{
+		UE_TRACE_ANALYSIS_DEBUG_Appendf("Serial=0x%07X ", EventDesc.Serial);
+	}
+	else
+	{
+		UE_TRACE_ANALYSIS_DEBUG_Appendf("Serial=%-9d ", EventDesc.Serial);
+	}
+	UE_TRACE_ANALYSIS_DEBUG_Appendf("Size=%d", Size);
+	if (EventDesc.bHasAux)
+	{
+		UE_TRACE_ANALYSIS_DEBUG_Append(" aux");
+	}
+	if (EventDesc.Uid == EKnownUids::AuxData)
+	{
+		UE_TRACE_ANALYSIS_DEBUG_Append(" data");
+	}
+	else if (EventDesc.Uid == EKnownUids::AuxDataTerminal)
+	{
+		UE_TRACE_ANALYSIS_DEBUG_Append(" end");
+	}
+	UE_TRACE_ANALYSIS_DEBUG_EndStringBuilder();
+}
+#endif // UE_TRACE_ANALYSIS_DEBUG
+
 ////////////////////////////////////////////////////////////////////////////////
 int32 FProtocol5Stage::ParseEvents(FStreamReader& Reader, EventDescArray& OutEventDescs, const FMachineContext& Context)
 {
 	while (!Reader.IsEmpty())
 	{
 		FEventDesc EventDesc;
+
 		int32 Size = ParseEvent(Reader, EventDesc, Context);
 		if (Size <= 0)
 		{
 			return Size;
 		}
 
-		TRACE_ANALYSIS_DEBUG("Event: %04d Uid:%04x Serial:%08x", OutEventDescs.Num(), EventDesc.Uid, EventDesc.Serial);
+#if UE_TRACE_ANALYSIS_DEBUG && UE_TRACE_ANALYSIS_DEBUG_LEVEL >= 3
+		PrintParsedEvent(OutEventDescs.Num(), EventDesc, Size);
+#endif // UE_TRACE_ANALYSIS_DEBUG
+
 		OutEventDescs.Add(EventDesc);
 
 		if (EventDesc.bHasAux)
@@ -3467,6 +4298,9 @@ int32 FProtocol5Stage::ParseEvents(FStreamReader& Reader, EventDescArray& OutEve
 
 			if (Ok == 0)
 			{
+#if UE_TRACE_ANALYSIS_DEBUG && UE_TRACE_ANALYSIS_DEBUG_LEVEL >= 2
+				UE_TRACE_ANALYSIS_DEBUG_LOG("Warning: Incomplete aux stack! Rewind %d parsed events.", OutEventDescs.Num() - RewindDescsNum);
+#endif
 				OutEventDescs.SetNum(RewindDescsNum);
 				Reader.RestoreMark(RewindMark);
 				break;
@@ -3508,7 +4342,9 @@ int32 FProtocol5Stage::ParseEventsWithAux(FStreamReader& Reader, EventDescArray&
 			return Size;
 		}
 
-		TRACE_ANALYSIS_DEBUG("Event: %04d Uid:%04x Serial:%08x", OutEventDescs.Num(), EventDesc.Uid, EventDesc.Serial);
+#if UE_TRACE_ANALYSIS_DEBUG && UE_TRACE_ANALYSIS_DEBUG_LEVEL >= 3
+		PrintParsedEvent(OutEventDescs.Num(), EventDesc, Size);
+#endif // UE_TRACE_ANALYSIS_DEBUG
 
 		Reader.Advance(Size);
 
@@ -3536,13 +4372,33 @@ int32 FProtocol5Stage::ParseEventsWithAux(FStreamReader& Reader, EventDescArray&
 				AuxKeyStack.Add(uint16(AuxKey));
 			}
 
-			// This event maybe in the middle of an earlier event's aux data blocks.
+			// This event may be in the middle of an earlier event's aux data blocks.
 			bUnsorted = true;
 		}
 
 		OutEventDescs.Add(EventDesc);
 
 		++AuxKey;
+
+		constexpr uint32 MaxAuxKey = 0x7fff;
+		if (AuxKeyStack.Num() == 1 && AuxKey > MaxAuxKey)
+		{
+			// If an "aux terminal" for the initial event was not detected after
+			// many intermediate events, we can assume it is lost.
+			check(FirstDescIndex > 0);
+			uint32 NumParsedEvents = uint32(OutEventDescs.Num()) - FirstDescIndex;
+			UE_TRACE_ANALYSIS_DEBUG_LOG("Error: Ignoring lost aux terminal for event with uid %u (desc %d), after parsing %u events",
+				OutEventDescs[FirstDescIndex - 1].Uid,
+				FirstDescIndex - 1,
+				NumParsedEvents);
+			Context.EmitMessagef(
+				EAnalysisMessageSeverity::Error,
+				TEXT("Ignoring lost aux terminal for event with uid %u, after parsing %u events."),
+				OutEventDescs[FirstDescIndex - 1].Uid,
+				NumParsedEvents);
+			AuxKeyStack.Pop();
+			break;
+		}
 	}
 
 	if (AuxKeyStack.Num() > 0)
@@ -3551,12 +4407,15 @@ int32 FProtocol5Stage::ParseEventsWithAux(FStreamReader& Reader, EventDescArray&
 		return 0;
 	}
 
-	checkf((AuxKey & 0xffff0000) == 0, TEXT("AuxKey overflow (%08x)"), AuxKey);
+	checkf((AuxKey & 0xffff0000) == 0, TEXT("AuxKey overflow (0x%X)"), AuxKey);
 
 	// Sort to get all aux-blocks contiguous with their owning event
 	if (bUnsorted)
 	{
 		uint32 NumDescs = OutEventDescs.Num() - FirstDescIndex;
+#if UE_TRACE_ANALYSIS_DEBUG && UE_TRACE_ANALYSIS_DEBUG_LEVEL >= 2
+		UE_TRACE_ANALYSIS_DEBUG_LOG("Sorting %u event descs", NumDescs);
+#endif
 		TArrayView<FEventDesc> DescsView(OutEventDescs.GetData() + FirstDescIndex, NumDescs);
 		Algo::StableSort(
 			DescsView,
@@ -3626,7 +4485,7 @@ int32 FProtocol5Stage::ParseEvent(FStreamReader& Reader, FEventDesc& EventDesc,
 		const FTypeRegistry::FTypeInfo* TypeInfo = TypeRegistry.Get(Uid);
 		if (TypeInfo == nullptr)
 		{
-			Context.EmitMessagef(EAnalysisMessageSeverity::Warning, TEXT("An unknown event UID (%u) was encountered."), Uid);
+			UE_TRACE_ANALYSIS_DEBUG_LOG("Warning: UID %u (0x%X) was not declared yet!", Uid, Uid);
 			return 0;
 		}
 
@@ -3646,7 +4505,11 @@ int32 FProtocol5Stage::ParseEvent(FStreamReader& Reader, FEventDesc& EventDesc,
 	EventDesc.Data = Cursor;
 
 	uint32 HeaderSize = uint32(UPTRINT(Cursor - Reader.GetPointer<uint8>()));
-	return HeaderSize + EventSize;
+	uint32 TotalEventSize = HeaderSize + EventSize;
+#if UE_TRACE_ANALYSIS_DEBUG
+	EventDesc.EventSize = TotalEventSize;
+#endif
+	return TotalEventSize;
 }
 
 ////////////////////////////////////////////////////////////////////////////////
@@ -3685,7 +4548,7 @@ void FProtocol5Stage::ForEachSerialGap(
 			}
 		}
 
-		// Consume consecutive events (including unsynchronised ones).
+		// Consume consecutive events (including unsynchronized ones).
 		Serial = EventDesc->Serial;
 		do
 		{
@@ -3713,13 +4576,13 @@ void FProtocol5Stage::ForEachSerialGap(
 ////////////////////////////////////////////////////////////////////////////////
 void FProtocol5Stage::DetectSerialGaps(TArray<FEventDescStream>& EventDescHeap)
 {
-	// Events that should be synchronised across threads are assigned serial
-	// numbers so they can be analysed in the correct order. Gaps in the
+	// Events that should be synchronized across threads are assigned serial
+	// numbers so they can be analyzed in the correct order. Gaps in the
 	// serials can occur under two scenarios; 1) when packets are dropped from
 	// the trace tail to make space for new trace events, and 2) when Trace's
 	// worker thread ticks, samples all the trace buffers and sends their data.
 	// In late-connect scenarios these gaps need to be skipped over in order to
-	// successfully reserialise events in the data stream. To further complicate
+	// successfully reserialize events in the data stream. To further complicate
 	// matters, most of the gaps from (2) will get filled by the following update,
 	// leading to initial false positive gaps. By embedding sync points in the
 	// stream we can reliably differentiate genuine gaps from temporary ones.
@@ -3733,6 +4596,7 @@ void FProtocol5Stage::DetectSerialGaps(TArray<FEventDescStream>& EventDescHeap)
 	}
 
 	SyncCount = Transport.GetSyncCount();
+	UE_TRACE_ANALYSIS_DEBUG_LOG("SyncCount: %d (%d previous serial gaps)", SyncCount, SerialGaps.Num());
 
 	if (SyncCount == 1)
 	{
@@ -3776,17 +4640,28 @@ void FProtocol5Stage::DetectSerialGaps(TArray<FEventDescStream>& EventDescHeap)
 			}
 
 			// If we're here something's probably gone wrong
+			UE_TRACE_ANALYSIS_DEBUG_LOG("Error: Serial gaps detection failed (SerialGap.Serial=%d, Lhs=%d)!", SerialGap.Serial, Lhs);
 			return false;
 		};
 
 		ForEachSerialGap(EventDescHeap, RecordGap);
 
+		UE_TRACE_ANALYSIS_DEBUG_LOG("Serial gaps: %d", GapCount);
+
 		if (GapCount == 0) //-V547
 		{
 			SerialGaps.Empty();
 			return;
 		}
 
+#if UE_TRACE_ANALYSIS_DEBUG && UE_TRACE_ANALYSIS_DEBUG_LEVEL >= 2
+		for (uint32 GapIndex = 0; GapIndex < GapCount; ++GapIndex)
+		{
+			const FEventDesc& SerialGap = SerialGaps[GapIndex];
+			UE_TRACE_ANALYSIS_DEBUG_LOG("  gap %u +%u", SerialGap.Serial, SerialGap.GapLength);
+		}
+#endif
+
 		// Turn the genuine gaps into a stream that DispatchEvents() can handle
 		// and use to skip over them.
 
@@ -3823,16 +4698,25 @@ bool FProtocol5Stage::DispatchKnownEvent(const FMachineContext& Context, uint32
 	switch (Uid)
 	{
 	case EKnownUids::EnterScope:
+#if UE_TRACE_ANALYSIS_DEBUG
+		Context.Bridge.DebugLogEnterScopeEvent(Uid, Cursor->EventSize + Cursor->AuxSize);
+#endif
 		Context.Bridge.EnterScope();
 		return true;
 
 	case EKnownUids::LeaveScope:
+#if UE_TRACE_ANALYSIS_DEBUG
+		Context.Bridge.DebugLogLeaveScopeEvent(Uid, Cursor->EventSize + Cursor->AuxSize);
+#endif
 		Context.Bridge.LeaveScope();
 		return true;
 
 	case EKnownUids::EnterScope_T:
 	{
 		uint64 RelativeTimestamp = *(uint64*)(Cursor->Data - 1) >> 8;
+#if UE_TRACE_ANALYSIS_DEBUG
+		Context.Bridge.DebugLogEnterScopeEvent(Uid, RelativeTimestamp, Cursor->EventSize + Cursor->AuxSize);
+#endif
 		Context.Bridge.EnterScope(RelativeTimestamp);
 		return true;
 	}
@@ -3840,6 +4724,9 @@ bool FProtocol5Stage::DispatchKnownEvent(const FMachineContext& Context, uint32
 	case EKnownUids::LeaveScope_T:
 	{
 		uint64 RelativeTimestamp = *(uint64*)(Cursor->Data - 1) >> 8;
+#if UE_TRACE_ANALYSIS_DEBUG
+		Context.Bridge.DebugLogLeaveScopeEvent(Uid, RelativeTimestamp, Cursor->EventSize + Cursor->AuxSize);
+#endif
 		Context.Bridge.LeaveScope(RelativeTimestamp);
 		return true;
 	}
@@ -3862,6 +4749,10 @@ int32 FProtocol5Stage::DispatchEvents(
 
 	FAuxDataCollector AuxCollector;
 
+#if UE_TRACE_ANALYSIS_DEBUG && UE_TRACE_ANALYSIS_DEBUG_LEVEL >= 2
+	UE_TRACE_ANALYSIS_DEBUG_LOG("Dispatch run of %u consecutive events (Tid=%u)", Count, Context.Bridge.GetActiveThreadId());
+#endif
+
 	for (const FEventDesc* Cursor = EventDesc, *End = EventDesc + Count; Cursor < End; ++Cursor)
 	{
 		uint32 Uid = uint32(Cursor->Uid);
@@ -3873,7 +4764,26 @@ int32 FProtocol5Stage::DispatchEvents(
 
 		if (!TypeRegistry.IsUidValid(Uid))
 		{
+			UE_TRACE_ANALYSIS_DEBUG_LOG("Warning: Unexpected event with UID %u (0x%X)", Uid, Uid);
 			Context.EmitMessagef(EAnalysisMessageSeverity::Warning, TEXT("An unknown event UID (%u) was encountered."), Uid);
+#if UE_TRACE_ANALYSIS_DEBUG && UE_TRACE_ANALYSIS_DEBUG_LEVEL >= 3
+			UE_TRACE_ANALYSIS_DEBUG_BeginStringBuilder();
+			const uint8* StartCursor = Cursor->Data;
+			for (uint32 i = 0; i < 32 && i < Cursor->EventSize; ++i)
+			{
+				UE_TRACE_ANALYSIS_DEBUG_Appendf("%02X ", StartCursor[i]);
+			}
+			UE_TRACE_ANALYSIS_DEBUG_EndStringBuilder();
+			UE_TRACE_ANALYSIS_DEBUG_ResetStringBuilder();
+			UE_TRACE_ANALYSIS_DEBUG_Append("[[[");
+			for (uint32 i = 0; i < 128 && i < Cursor->EventSize; ++i)
+			{
+				UE_TRACE_ANALYSIS_DEBUG_AppendChar((char)StartCursor[i]);
+			}
+			UE_TRACE_ANALYSIS_DEBUG_Append("]]]");
+			UE_TRACE_ANALYSIS_DEBUG_EndStringBuilder();
+#endif // UE_TRACE_ANALYSIS_DEBUG
+
 			return -1;
 		}
 
@@ -3886,6 +4796,12 @@ int32 FProtocol5Stage::DispatchEvents(
 			TypeInfo->EventSize,
 		};
 
+#if UE_TRACE_ANALYSIS_DEBUG
+		uint32 FixedSize = Cursor->EventSize;
+		uint32 AuxSize = Cursor->AuxSize;
+		const FEventDesc* EventCursor = Cursor;
+#endif
+
 		// Gather its auxiliary data blocks into a collector.
 		if (Cursor->bHasAux)
 		{
@@ -3907,9 +4823,20 @@ int32 FProtocol5Stage::DispatchEvents(
 				AuxData.FieldIndex = AuxHeader->FieldIndex_Size & FAuxHeader::FieldMask;
 				// AuxData.FieldSizeAndType = ... - this is assigned on demand in GetData()
 				AuxCollector.Add(AuxData);
+
+#if UE_TRACE_ANALYSIS_DEBUG
+				AuxSize += Cursor->EventSize + Cursor->AuxSize;
+#endif
 			}
+#if UE_TRACE_ANALYSIS_DEBUG
+			AuxSize += 1; // for AuxDataTerminal
+#endif
 		}
 
+#if UE_TRACE_ANALYSIS_DEBUG
+		Context.Bridge.DebugLogEvent(TypeInfo, FixedSize, AuxSize, EventCursor->Serial);
+#endif
+
 		Context.Bridge.OnEvent(EventDataInfo);
 
 		AuxCollector.Reset();
@@ -3923,10 +4850,12 @@ int32 FProtocol5Stage::DispatchEvents(
 // {{{1 protocol-6 -------------------------------------------------------------
 
 ////////////////////////////////////////////////////////////////////////////////
-class FProtocol6Stage : public FProtocol5Stage
+class FProtocol6Stage
+	: public FProtocol5Stage
 {
 public:
-							FProtocol6Stage(FTransport* InTransport);
+						FProtocol6Stage(FTransport* InTransport);
+	virtual void		ExitStage(const FMachineContext& Context) override;
 };
 
 ////////////////////////////////////////////////////////////////////////////////
@@ -3936,21 +4865,37 @@ FProtocol6Stage::FProtocol6Stage(FTransport* InTransport)
 	EventVersion = 6;
 }
 
+////////////////////////////////////////////////////////////////////////////////
+void FProtocol6Stage::ExitStage(const FMachineContext& Context)
+{
+	// Ensure the transport does not have pending buffers (i.e. event data not yet processed).
+	if (!Transport.IsEmpty())
+	{
+		Context.EmitMessage(EAnalysisMessageSeverity::Warning, TEXT("Transport buffers are not empty at end of analysis (protocol 6)!"));
+	}
+
+#if UE_TRACE_ANALYSIS_DEBUG
+	Context.Bridge.GetSerial().Value = NextSerial;
+	Transport.DebugEnd();
+#endif
+}
+
 
 
 // {{{1 protocol-7 -------------------------------------------------------------
 
 ////////////////////////////////////////////////////////////////////////////////
-class FProtocol7Stage : public FProtocol6Stage
+class FProtocol7Stage
+	: public FProtocol6Stage
 {
 public:
-							FProtocol7Stage(FTransport* InTransport);
-
-	virtual void			SetSizeIfKnownEvent(uint32 Uid, uint32& InOutEventSize) override;
-	virtual bool			DispatchKnownEvent(const FMachineContext& Context, uint32 Uid, const FEventDesc* Cursor) override;
+						FProtocol7Stage(FTransport* InTransport);
+	virtual void		ExitStage(const FMachineContext& Context) override;
+	virtual void		SetSizeIfKnownEvent(uint32 Uid, uint32& InOutEventSize) override;
+	virtual bool		DispatchKnownEvent(const FMachineContext& Context, uint32 Uid, const FEventDesc* Cursor) override;
 
 protected:
-	using EKnownUids		= Protocol7::EKnownEventUids;
+	using EKnownUids	= Protocol7::EKnownEventUids;
 };
 
 ////////////////////////////////////////////////////////////////////////////////
@@ -3960,6 +4905,21 @@ FProtocol7Stage::FProtocol7Stage(FTransport* InTransport)
 	EventVersion = 7;
 }
 
+////////////////////////////////////////////////////////////////////////////////
+void FProtocol7Stage::ExitStage(const FMachineContext& Context)
+{
+	// Ensure the transport does not have pending buffers (i.e. event data not yet processed).
+	if (!Transport.IsEmpty())
+	{
+		Context.EmitMessage(EAnalysisMessageSeverity::Warning, TEXT("Transport buffers are not empty at end of analysis (protocol 7)!"));
+	}
+
+#if UE_TRACE_ANALYSIS_DEBUG
+	Context.Bridge.GetSerial().Value = NextSerial;
+	Transport.DebugEnd();
+#endif
+}
+
 ////////////////////////////////////////////////////////////////////////////////
 void FProtocol7Stage::SetSizeIfKnownEvent(uint32 Uid, uint32& InOutEventSize)
 {
@@ -3984,16 +4944,25 @@ bool FProtocol7Stage::DispatchKnownEvent(const FMachineContext& Context, uint32
 	switch (Uid)
 	{
 	case EKnownUids::EnterScope:
+#if UE_TRACE_ANALYSIS_DEBUG
+		Context.Bridge.DebugLogEnterScopeEvent(Uid, Cursor->EventSize + Cursor->AuxSize);
+#endif
 		Context.Bridge.EnterScope();
 		return true;
 
 	case EKnownUids::LeaveScope:
+#if UE_TRACE_ANALYSIS_DEBUG
+		Context.Bridge.DebugLogLeaveScopeEvent(Uid, Cursor->EventSize + Cursor->AuxSize);
+#endif
 		Context.Bridge.LeaveScope();
 		return true;
 
 	case EKnownUids::EnterScope_TA:
 	{
 		uint64 AbsoluteTimestamp = *(uint64*)Cursor->Data;
+#if UE_TRACE_ANALYSIS_DEBUG
+		Context.Bridge.DebugLogEnterScopeAEvent(Uid, AbsoluteTimestamp, Cursor->EventSize + Cursor->AuxSize);
+#endif
 		Context.Bridge.EnterScopeA(AbsoluteTimestamp);
 		return true;
 	}
@@ -4001,6 +4970,9 @@ bool FProtocol7Stage::DispatchKnownEvent(const FMachineContext& Context, uint32
 	case EKnownUids::LeaveScope_TA:
 	{
 		uint64 AbsoluteTimestamp = *(uint64*)Cursor->Data;
+#if UE_TRACE_ANALYSIS_DEBUG
+		Context.Bridge.DebugLogLeaveScopeAEvent(Uid, AbsoluteTimestamp, Cursor->EventSize + Cursor->AuxSize);
+#endif
 		Context.Bridge.LeaveScopeA(AbsoluteTimestamp);
 		return true;
 	}
@@ -4008,6 +4980,9 @@ bool FProtocol7Stage::DispatchKnownEvent(const FMachineContext& Context, uint32
 	case EKnownUids::EnterScope_TB:
 	{
 		uint64 BaseRelativeTimestamp = *(uint64*)(Cursor->Data - 1) >> 8;
+#if UE_TRACE_ANALYSIS_DEBUG
+		Context.Bridge.DebugLogEnterScopeBEvent(Uid, BaseRelativeTimestamp, Cursor->EventSize + Cursor->AuxSize);
+#endif
 		Context.Bridge.EnterScopeB(BaseRelativeTimestamp);
 		return true;
 	}
@@ -4015,6 +4990,9 @@ bool FProtocol7Stage::DispatchKnownEvent(const FMachineContext& Context, uint32
 	case EKnownUids::LeaveScope_TB:
 	{
 		uint64 BaseRelativeTimestamp = *(uint64*)(Cursor->Data - 1) >> 8;
+#if UE_TRACE_ANALYSIS_DEBUG
+		Context.Bridge.DebugLogLeaveScopeBEvent(Uid, BaseRelativeTimestamp, Cursor->EventSize + Cursor->AuxSize);
+#endif
 		Context.Bridge.LeaveScopeB(BaseRelativeTimestamp);
 		return true;
 	}
@@ -4055,10 +5033,10 @@ FEstablishTransportStage::EStatus FEstablishTransportStage::OnData(
 	}
 
 	uint32 TransportVersion = Header->TransportVersion;
-	TRACE_ANALYSIS_DEBUG("TransportVersion: %u", TransportVersion);
+	UE_TRACE_ANALYSIS_DEBUG_LOG("TransportVersion: %u", TransportVersion);
 
 	uint32 ProtocolVersion = Header->ProtocolVersion;
-	TRACE_ANALYSIS_DEBUG("ProtocolVersion: %u", ProtocolVersion);
+	UE_TRACE_ANALYSIS_DEBUG_LOG("ProtocolVersion: %u", ProtocolVersion);
 
 	FTransport* Transport = nullptr;
 	switch (TransportVersion)
@@ -4069,6 +5047,7 @@ FEstablishTransportStage::EStatus FEstablishTransportStage::OnData(
 	case ETransport::TidPacketSync:	Transport = new FTidPacketTransportSync(); break;
 	default:
 		{
+			UE_TRACE_ANALYSIS_DEBUG_LOG("Error: Invalid transport version %u", TransportVersion);
 			Context.EmitMessagef(
 				EAnalysisMessageSeverity::Error,
 				TEXT("Unknown transport version: %u. You may need to recompile this application"),
@@ -4078,6 +5057,10 @@ FEstablishTransportStage::EStatus FEstablishTransportStage::OnData(
 		}
 	}
 
+#if UE_TRACE_ANALYSIS_DEBUG
+	Transport->DebugBegin();
+#endif
+
 	switch (ProtocolVersion)
 	{
 	case Protocol0::EProtocol::Id:
@@ -4114,6 +5097,7 @@ FEstablishTransportStage::EStatus FEstablishTransportStage::OnData(
 
 	default:
 		{
+			UE_TRACE_ANALYSIS_DEBUG_LOG("Error: Invalid protocol version %u", ProtocolVersion);
 			Context.EmitMessagef(
 				EAnalysisMessageSeverity::Error,
 				TEXT("Unknown protocol version: %u. You may need to recompile this application"),
@@ -4123,7 +5107,14 @@ FEstablishTransportStage::EStatus FEstablishTransportStage::OnData(
 		}
 	}
 
+	UE_TRACE_ANALYSIS_DEBUG_LOG("");
+
 	Reader.Advance(sizeof(*Header));
+
+#if UE_TRACE_ANALYSIS_DEBUG_API
+	Context.Bridge.OnVersion(TransportVersion, ProtocolVersion);
+#endif // UE_TRACE_ANALYSIS_DEBUG_API
+
 	return EStatus::Continue;
 }
 
@@ -4156,6 +5147,27 @@ FMetadataStage::EStatus FMetadataStage::OnData(
 		return EStatus::NotEnoughData;
 	}
 
+#if UE_TRACE_ANALYSIS_DEBUG
+	UE_TRACE_ANALYSIS_DEBUG_BeginStringBuilder();
+	UE_TRACE_ANALYSIS_DEBUG_Appendf("Metadata: %u + %u bytes (", uint32(sizeof(*MetadataSize)), uint32(*MetadataSize));
+	const uint32 PrintByteCount = FMath::Min(32u, uint32(*MetadataSize));
+	Metadata += sizeof(*MetadataSize);
+	for (uint32 Index = 0; Index < PrintByteCount; ++Index, ++Metadata)
+	{
+		if (Index != 0)
+		{
+			UE_TRACE_ANALYSIS_DEBUG_AppendChar(' ');
+		}
+		UE_TRACE_ANALYSIS_DEBUG_Appendf("%02X", uint32(*Metadata));
+	}
+	if (PrintByteCount != uint32(*MetadataSize))
+	{
+		UE_TRACE_ANALYSIS_DEBUG_Append("...");
+	}
+	UE_TRACE_ANALYSIS_DEBUG_AppendChar(')');
+	UE_TRACE_ANALYSIS_DEBUG_EndStringBuilder();
+#endif // UE_TRACE_ANALYSIS_DEBUG
+
 	Reader.Advance(sizeof(*MetadataSize) + *MetadataSize);
 
 	Context.Machine.QueueStage<FEstablishTransportStage>();
@@ -4191,6 +5203,7 @@ FMagicStage::EStatus FMagicStage::OnData(
 	if (Magic == 'ECRT' || Magic == '2CRT')
 	{
 		// Source is big-endian which we don't currently support
+		UE_TRACE_ANALYSIS_DEBUG_LOG("Error: Invalid magic header (big-endian not supported)");
 		Context.EmitMessage(EAnalysisMessageSeverity::Error, TEXT("Big endian traces are currently not supported."));
 		return EStatus::Error;
 	}
@@ -4220,6 +5233,7 @@ FMagicStage::EStatus FMagicStage::OnData(
 		return EStatus::Continue;
 	}
 
+	UE_TRACE_ANALYSIS_DEBUG_LOG("Error: Invalid magic header");
 	Context.EmitMessage(EAnalysisMessageSeverity::Error, TEXT("The file or stream was not recognized as trace stream."));
 	return EStatus::Error;
 }
@@ -4250,6 +5264,8 @@ FAnalysisEngine::FImpl::FImpl(TArray<IAnalyzer*>&& Analyzers, FMessageDelegate&&
 ////////////////////////////////////////////////////////////////////////////////
 void FAnalysisEngine::FImpl::Begin()
 {
+	UE_TRACE_ANALYSIS_DEBUG_LOG("FAnalysisEngine::Begin()");
+
 	Machine.QueueStage<FMagicStage>();
 	Machine.Transition();
 }
@@ -4258,7 +5274,107 @@ void FAnalysisEngine::FImpl::Begin()
 void FAnalysisEngine::FImpl::End()
 {
 	Machine.Transition();
+
+#if UE_TRACE_ANALYSIS_DEBUG
+	FAnalysisState& State = Bridge.GetState();
+
+	UE_TRACE_ANALYSIS_DEBUG_LOG("");
+	UE_TRACE_ANALYSIS_DEBUG_LOG("SerialWrappedCount: %d", State.SerialWrappedCount);
+	UE_TRACE_ANALYSIS_DEBUG_LOG("Serial.Value: %u (0x%X)", State.Serial.Value, State.Serial.Value);
+	UE_TRACE_ANALYSIS_DEBUG_LOG("MaxEventDescs: %d", State.MaxEventDescs);
+	UE_TRACE_ANALYSIS_DEBUG_LOG("SkippedSerialGaps: %d", State.NumSkippedSerialGaps);
+
+	uint32 Digits = 0;
+	uint64 Value = State.TotalEventSize;
+	do
+	{
+		Value /= 10;
+		++Digits;
+	} while (Value != 0);
+
+	uint64 ScopeEventCount = State.EnterScopeEventCount + State.LeaveScopeEventCount;
+	uint64 ScopeTEventCount = State.EnterScopeTEventCount + State.LeaveScopeTEventCount;
+
+	const int32 NC = Digits;
+	UE_TRACE_ANALYSIS_DEBUG_LOG("");
+	UE_TRACE_ANALYSIS_DEBUG_LOG("TotalEventCount:%*llu events", NC, State.TotalEventCount);
+	if (State.TotalEventCount > 0)
+	{
+		UE_TRACE_ANALYSIS_DEBUG_LOG("       NewEvent:%*llu events (%.1f%%)", NC, State.NewEventCount, (double)State.NewEventCount / (double)State.TotalEventCount * 100.0);
+		UE_TRACE_ANALYSIS_DEBUG_LOG("           Sync:%*llu events (%.1f%%)", NC, State.SyncEventCount, (double)State.SyncEventCount / (double)State.TotalEventCount * 100.0);
+		UE_TRACE_ANALYSIS_DEBUG_LOG("ImportantNoSync:%*llu events (%.1f%%)", NC, State.ImportantNoSyncEventCount, (double)State.ImportantNoSyncEventCount / (double)State.TotalEventCount * 100.0);
+		UE_TRACE_ANALYSIS_DEBUG_LOG("    OtherNoSync:%*llu events (%.1f%%)", NC, State.OtherNoSyncEventCount, (double)State.OtherNoSyncEventCount / (double)State.TotalEventCount * 100.0);
+		UE_TRACE_ANALYSIS_DEBUG_LOG("          Scope:%*llu events (%.1f%%) = %llu enter + %llu leave",
+			NC, ScopeEventCount,
+			(double)ScopeEventCount / (double)State.TotalEventCount * 100.0,
+			State.EnterScopeEventCount,
+			State.LeaveScopeEventCount);
+		UE_TRACE_ANALYSIS_DEBUG_LOG("        Scope_T:%*llu events (%.1f%%) = %llu enter + %llu leave",
+			NC, ScopeTEventCount,
+			(double)ScopeTEventCount / (double)State.TotalEventCount * 100.0,
+			State.EnterScopeTEventCount,
+			State.LeaveScopeTEventCount);
+		const int64 CountError =
+			State.TotalEventCount
+			- State.NewEventCount
+			- State.SyncEventCount
+			- State.ImportantNoSyncEventCount
+			- State.OtherNoSyncEventCount
+			- ScopeEventCount
+			- ScopeTEventCount;
+		if (CountError != 0)
+		{
+			UE_TRACE_ANALYSIS_DEBUG_LOG("          error: %lld !!!", CountError);
+		}
+		uint64 DispatchedEventCount = State.TotalEventCount - State.NewEventCount - State.EnterScopeEventCount - State.EnterScopeTEventCount;
+		uint64 NormalEventCount = State.SyncEventCount + State.ImportantNoSyncEventCount + State.OtherNoSyncEventCount - State.EnterScopeEventCount - State.EnterScopeTEventCount;
+		check(DispatchedEventCount == NormalEventCount + ScopeEventCount + ScopeTEventCount);
+		UE_TRACE_ANALYSIS_DEBUG_LOG("     Dispatched:%*llu events = %llu normal + %llu scoped", NC, DispatchedEventCount, NormalEventCount, ScopeEventCount + ScopeTEventCount);
+		UE_TRACE_ANALYSIS_DEBUG_LOG("(normal dispatched events = sync events + no sync events - scope enter events)");
+		UE_TRACE_ANALYSIS_DEBUG_LOG("");
+	}
+
+	uint64 ScopeEventSize = State.EnterScopeEventSize + State.LeaveScopeEventSize;
+	uint64 ScopeTEventSize = State.EnterScopeTEventSize + State.LeaveScopeTEventSize;
+
+	const int32 NS = Digits + 1;
+	UE_TRACE_ANALYSIS_DEBUG_LOG(" TotalEventSize:%*llu bytes (%.1f MiB)", NS, State.TotalEventSize, (double)State.TotalEventSize / (1024.0 * 1024.0));
+	if (State.TotalEventSize > 0)
+	{
+		UE_TRACE_ANALYSIS_DEBUG_LOG("       NewEvent:%*llu bytes (%.1f%%)", NS, State.NewEventSize, (double)State.NewEventSize / (double)State.TotalEventSize * 100.0);
+		UE_TRACE_ANALYSIS_DEBUG_LOG("           Sync:%*llu bytes (%.1f%%)", NS, State.SyncEventSize, (double)State.SyncEventSize / (double)State.TotalEventSize * 100.0);
+		UE_TRACE_ANALYSIS_DEBUG_LOG("ImportantNoSync:%*llu bytes (%.1f%%)", NS, State.ImportantNoSyncEventSize, (double)State.ImportantNoSyncEventSize / (double)State.TotalEventSize * 100.0);
+		UE_TRACE_ANALYSIS_DEBUG_LOG("    OtherNoSync:%*llu bytes (%.1f%%)", NS, State.OtherNoSyncEventSize, (double)State.OtherNoSyncEventSize / (double)State.TotalEventSize * 100.0);
+		UE_TRACE_ANALYSIS_DEBUG_LOG("          Scope:%*llu bytes (%.1f%%) = %llu enter + %llu leave",
+			NS, ScopeEventSize,
+			(double)ScopeEventSize / (double)State.TotalEventSize * 100.0,
+			State.EnterScopeEventSize,
+			State.LeaveScopeEventSize);
+		UE_TRACE_ANALYSIS_DEBUG_LOG("        Scope_T:%*llu bytes (%.1f%%) = %llu enter + %llu leave",
+			NS, ScopeTEventSize,
+			(double)ScopeTEventSize / (double)State.TotalEventSize * 100.0,
+			State.EnterScopeTEventSize,
+			State.LeaveScopeTEventSize);
+		const int64 SizeError =
+			State.TotalEventSize
+			- State.NewEventSize
+			- State.SyncEventSize
+			- State.ImportantNoSyncEventSize
+			- State.OtherNoSyncEventSize
+			- ScopeEventSize
+			- ScopeTEventSize;
+		if (SizeError != 0)
+		{
+			UE_TRACE_ANALYSIS_DEBUG_LOG("          error: %lld !!!", SizeError);
+		}
+		UE_TRACE_ANALYSIS_DEBUG_LOG("     Dispatched:%*llu bytes (Total - NewEvent)", NS, State.TotalEventSize - State.NewEventSize);
+	}
+#endif // UE_TRACE_ANALYSIS_DEBUG
+
 	Bridge.Reset();
+
+	UE_TRACE_ANALYSIS_DEBUG_LOG("");
+	UE_TRACE_ANALYSIS_DEBUG_LOG("FAnalysisEngine::End()");
 }
 
 ////////////////////////////////////////////////////////////////////////////////
@@ -4305,7 +5421,6 @@ bool FAnalysisEngine::OnData(FStreamReader& Reader)
 } // namespace Trace
 } // namespace UE
 
-#undef TRACE_ANALYSIS_DEBUG
 #undef LOCTEXT_NAMESPACE
 
 /* vim: set foldlevel=1 : */
diff --git a/Engine/Source/Developer/TraceAnalysis/Private/Analysis/StreamReader.cpp b/Engine/Source/Developer/TraceAnalysis/Private/Analysis/StreamReader.cpp
--- a/Engine/Source/Developer/TraceAnalysis/Private/Analysis/StreamReader.cpp
+++ b/Engine/Source/Developer/TraceAnalysis/Private/Analysis/StreamReader.cpp
@@ -25,15 +25,16 @@ void FStreamReader::Advance(uint32 Size)
 }
 
 ////////////////////////////////////////////////////////////////////////////////
-int32 FStreamReader::GetRemaining() const
+uint32 FStreamReader::GetRemaining() const
 {
-	return int32(PTRINT(End - Cursor));
+	check(End >= Cursor);
+	return End - Cursor;
 }
 
 ////////////////////////////////////////////////////////////////////////////////
 bool FStreamReader::CanMeetDemand() const
 {
-	return (GetRemaining() >= int32(DemandHint));
+	return GetRemaining() >= DemandHint;
 }
 
 ////////////////////////////////////////////////////////////////////////////////
@@ -101,7 +102,8 @@ uint8* FStreamBuffer::Append(uint32 Size)
 ////////////////////////////////////////////////////////////////////////////////
 void FStreamBuffer::Consolidate()
 {
-	int32 Remaining = GetRemaining();
+	uint32 Remaining = GetRemaining();
+	check((uint64)DemandHint + (uint64)Remaining < (1ull << 31));
 	DemandHint += Remaining;
 
 	if (DemandHint >= BufferSize)
diff --git a/Engine/Source/Developer/TraceAnalysis/Private/Analysis/StreamReader.h b/Engine/Source/Developer/TraceAnalysis/Private/Analysis/StreamReader.h
--- a/Engine/Source/Developer/TraceAnalysis/Private/Analysis/StreamReader.h
+++ b/Engine/Source/Developer/TraceAnalysis/Private/Analysis/StreamReader.h
@@ -21,7 +21,7 @@ public:
 	const uint8*				GetPointerUnchecked();
 	void						Advance(uint32 Size);
 	bool						IsEmpty() const;
-	int32						GetRemaining() const;
+	uint32						GetRemaining() const;
 	bool						CanMeetDemand() const;
 	bool						Backtrack(const uint8* To);
 	struct FMark*				SaveMark() const;
@@ -72,6 +72,7 @@ public:
 	int32						Fill(Lambda&& Source);
 	void						Append(const uint8* Data, uint32 Size);
 	uint8*						Append(uint32 Size);
+	uint32						GetBufferSize() const { return BufferSize; }
 
 protected:
 	void						Consolidate();
diff --git a/Engine/Source/Developer/TraceAnalysis/Private/Analysis/Transport/PacketTransport.cpp b/Engine/Source/Developer/TraceAnalysis/Private/Analysis/Transport/PacketTransport.cpp
--- a/Engine/Source/Developer/TraceAnalysis/Private/Analysis/Transport/PacketTransport.cpp
+++ b/Engine/Source/Developer/TraceAnalysis/Private/Analysis/Transport/PacketTransport.cpp
@@ -3,6 +3,8 @@
 #include "PacketTransport.h"
 #include "HAL/UnrealMemory.h"
 
+#include "TraceAnalysisDebug.h"
+
 #include <initializer_list>
 
 ////////////////////////////////////////////////////////////////////////////////
@@ -56,6 +58,12 @@ void FPacketTransport::Advance(uint32 BlockSize)
 	}
 }
 
+////////////////////////////////////////////////////////////////////////////////
+bool FPacketTransport::IsEmpty() const
+{
+	return (ActiveList == nullptr) || (ActiveList->Cursor + 1 > ActiveList->Size);
+}
+
 ////////////////////////////////////////////////////////////////////////////////
 const uint8* FPacketTransport::GetPointerImpl(uint32 BlockSize)
 {
@@ -177,5 +185,26 @@ bool FPacketTransport::GetNextBatch()
 	return true;
 }
 
+////////////////////////////////////////////////////////////////////////////////
+void FPacketTransport::DebugBegin()
+{
+#if UE_TRACE_ANALYSIS_DEBUG
+	UE_TRACE_ANALYSIS_DEBUG_LOG("FPacketTransport::DebugBegin()");
+#endif // UE_TRACE_ANALYSIS_DEBUG
+}
+
+////////////////////////////////////////////////////////////////////////////////
+void FPacketTransport::DebugEnd()
+{
+#if UE_TRACE_ANALYSIS_DEBUG
+	if (!IsEmpty())
+	{
+		UE_TRACE_ANALYSIS_DEBUG_LOG("Error: FPacketTransport is not empty!");
+	}
+	UE_TRACE_ANALYSIS_DEBUG_LOG("FPacketTransport::DebugEnd()");
+#endif // UE_TRACE_ANALYSIS_DEBUG
+}
+
+////////////////////////////////////////////////////////////////////////////////
 } // namespace Trace
 } // namespace UE
diff --git a/Engine/Source/Developer/TraceAnalysis/Private/Analysis/Transport/PacketTransport.h b/Engine/Source/Developer/TraceAnalysis/Private/Analysis/Transport/PacketTransport.h
--- a/Engine/Source/Developer/TraceAnalysis/Private/Analysis/Transport/PacketTransport.h
+++ b/Engine/Source/Developer/TraceAnalysis/Private/Analysis/Transport/PacketTransport.h
@@ -14,6 +14,11 @@ class FPacketTransport
 public:
 	virtual					~FPacketTransport();
 	virtual void			Advance(uint32 BlockSize) override;
+	virtual bool			IsEmpty() const override;
+	virtual void			DebugBegin() override;
+	virtual void			DebugEnd() override;
+
+protected:
 	virtual const uint8*	GetPointerImpl(uint32 BlockSize) override;
 
 private:
diff --git a/Engine/Source/Developer/TraceAnalysis/Private/Analysis/Transport/TidPacketTransport.cpp b/Engine/Source/Developer/TraceAnalysis/Private/Analysis/Transport/TidPacketTransport.cpp
--- a/Engine/Source/Developer/TraceAnalysis/Private/Analysis/Transport/TidPacketTransport.cpp
+++ b/Engine/Source/Developer/TraceAnalysis/Private/Analysis/Transport/TidPacketTransport.cpp
@@ -21,6 +21,22 @@ TRACELOG_API int32 Decode(const void*, int32, void*, int32);
 namespace UE {
 namespace Trace {
 
+////////////////////////////////////////////////////////////////////////////////
+bool FTidPacketTransport::IsEmpty() const
+{
+	if (Threads.Num() > 0)
+	{
+		for (const FThreadStream& Thread : Threads)
+		{
+			if (!Thread.Buffer.IsEmpty())
+			{
+				return false;
+			}
+		}
+	}
+	return true;
+}
+
 ////////////////////////////////////////////////////////////////////////////////
 FTidPacketTransport::EReadPacketResult FTidPacketTransport::ReadPacket()
 {
@@ -37,6 +53,12 @@ FTidPacketTransport::EReadPacketResult FTidPacketTransport::ReadPacket()
 		return EReadPacketResult::NeedMoreData;
 	}
 
+#if UE_TRACE_ANALYSIS_DEBUG
+	++NumPackets;
+	TotalPacketHeaderSize += sizeof(FTidPacketBase);
+	TotalPacketSize += PacketBase->PacketSize;
+#endif // UE_TRACE_ANALYSIS_DEBUG
+
 	FTransport::Advance(PacketBase->PacketSize);
 
 	uint32 ThreadId = PacketBase->ThreadId & FTidPacketBase::ThreadIdMask;
@@ -44,6 +66,9 @@ FTidPacketTransport::EReadPacketResult FTidPacketTransport::ReadPacket()
 	if (ThreadId == ETransportTid::Sync)
 	{
 		++Synced;
+#if UE_TRACE_ANALYSIS_DEBUG
+		UE_TRACE_ANALYSIS_DEBUG_LOG("[SYNC %u]", Synced);
+#endif // UE_TRACE_ANALYSIS_DEBUG
 		return EReadPacketResult::NeedMoreData;	// Do not read any more packets. Gives consumers a
 													// chance to sample the world at each known sync point.
 	}
@@ -60,6 +85,24 @@ FTidPacketTransport::EReadPacketResult FTidPacketTransport::ReadPacket()
 	{
 		const auto* Packet = (const FTidPacketEncoded*)PacketBase;
 		uint16 DecodedSize = Packet->DecodedSize;
+
+#if UE_TRACE_ANALYSIS_DEBUG
+#if UE_TRACE_ANALYSIS_DEBUG_LEVEL >= 3
+		UE_TRACE_ANALYSIS_DEBUG_LOG("[PACKET %u] Tid=%u, Size: %u + %u bytes, DecodedSize: %u bytes (%.0f%%)%s",
+			NumPackets,
+			ThreadId,
+			uint32(sizeof(FTidPacketEncoded)),
+			DataSize,
+			DecodedSize,
+			((double)DataSize * 100.0) / (double)DecodedSize - 100.0,
+			DataSize > DecodedSize ? " !!!" : "");
+#endif // UE_TRACE_ANALYSIS_DEBUG_LEVEL
+		TotalPacketHeaderSize += sizeof(FTidPacketEncoded) - sizeof(FTidPacketBase);
+		TotalDecodedSize += DecodedSize;
+		uint64& TotalDataSizePerThread = DataSizePerThread.FindOrAdd(ThreadId);
+		TotalDataSizePerThread += DecodedSize;
+#endif // UE_TRACE_ANALYSIS_DEBUG
+
 		uint8* Dest = Thread->Buffer.Append(DecodedSize);
 		DataSize -= sizeof(DecodedSize);
 		int32 ResultSize = UE::Trace::Private::Decode(Packet->Data, DataSize, Dest, DecodedSize);
@@ -68,10 +111,42 @@ FTidPacketTransport::EReadPacketResult FTidPacketTransport::ReadPacket()
 			UE_LOG(LogCore, Error, TEXT("Unable to decompress packet, expected %d bytes but decoded %d bytes."), DecodedSize, ResultSize);
 			return EReadPacketResult::ReadError;
 		}
+
+#if UE_TRACE_ANALYSIS_DEBUG && UE_TRACE_ANALYSIS_DEBUG_LEVEL >= 4
+		UE_TRACE_ANALYSIS_DEBUG_BeginStringBuilder();
+		for (uint32 i = 0; i < 32 && i < DecodedSize; ++i)
+		{
+			UE_TRACE_ANALYSIS_DEBUG_Appendf("%02X ", Dest[i]);
+		}
+		UE_TRACE_ANALYSIS_DEBUG_EndStringBuilder();
+#endif // UE_TRACE_ANALYSIS_DEBUG
 	}
 	else
 	{
+#if UE_TRACE_ANALYSIS_DEBUG
+#if UE_TRACE_ANALYSIS_DEBUG_LEVEL >= 3
+		UE_TRACE_ANALYSIS_DEBUG_LOG("[PACKET %u] Tid=%u, Size: %u + %u bytes",
+			NumPackets,
+			ThreadId,
+			uint32(sizeof(FTidPacket)),
+			DataSize);
+#endif // UE_TRACE_ANALYSIS_DEBUG_LEVEL
+		TotalPacketHeaderSize += sizeof(FTidPacket) - sizeof(FTidPacketBase);
+		TotalDecodedSize += DataSize;
+		uint64& TotalDataSizePerThread = DataSizePerThread.FindOrAdd(ThreadId);
+		TotalDataSizePerThread += DataSize;
+#endif // UE_TRACE_ANALYSIS_DEBUG
+
 		Thread->Buffer.Append((uint8*)(PacketBase + 1), DataSize);
+
+#if UE_TRACE_ANALYSIS_DEBUG && UE_TRACE_ANALYSIS_DEBUG_LEVEL >= 4
+		UE_TRACE_ANALYSIS_DEBUG_BeginStringBuilder();
+		for (uint32 i = 0; i < 32 && i < DataSize; ++i)
+		{
+			UE_TRACE_ANALYSIS_DEBUG_Appendf("%02X ", ((uint8*)(PacketBase + 1))[i]);
+		}
+		UE_TRACE_ANALYSIS_DEBUG_EndStringBuilder();
+#endif // UE_TRACE_ANALYSIS_DEBUG
 	}
 
 	return EReadPacketResult::Continue;
@@ -116,9 +191,121 @@ FTidPacketTransport::ETransportResult FTidPacketTransport::Update()
 		return (Thread.ThreadId <= ETransportTid::Importants) ? false : Thread.Buffer.IsEmpty();
 	});
 
+#if UE_TRACE_ANALYSIS_DEBUG
+	DebugUpdate();
+#endif // UE_TRACE_ANALYSIS_DEBUG
+
 	return Result == EReadPacketResult::ReadError ? ETransportResult::Error : ETransportResult::Ok;
 }
 
+////////////////////////////////////////////////////////////////////////////////
+void FTidPacketTransport::DebugUpdate()
+{
+#if UE_TRACE_ANALYSIS_DEBUG
+	const int32 ThreadCount = Threads.Num();
+	for (int32 ThreadIndex = 0; ThreadIndex < ThreadCount; ++ThreadIndex)
+	{
+		const FThreadStream& Thread = Threads[ThreadIndex];
+		const uint32 BufferSize = Thread.Buffer.GetBufferSize();
+		const uint32 DataSize = Thread.Buffer.GetRemaining();
+
+		if (BufferSize > MaxBufferSize)
+		{
+			MaxBufferSize = BufferSize;
+			MaxBufferSizeThreadId = Thread.ThreadId;
+		}
+
+		if (DataSize > MaxDataSizePerBuffer)
+		{
+			MaxDataSizePerBuffer = DataSize;
+			MaxDataSizePerBufferThreadId = Thread.ThreadId;
+		}
+	}
+#endif // UE_TRACE_ANALYSIS_DEBUG
+}
+
+////////////////////////////////////////////////////////////////////////////////
+void FTidPacketTransport::DebugBegin()
+{
+#if UE_TRACE_ANALYSIS_DEBUG
+	UE_TRACE_ANALYSIS_DEBUG_LOG("FTidPacketTransport::DebugBegin()");
+#endif // UE_TRACE_ANALYSIS_DEBUG
+}
+
+////////////////////////////////////////////////////////////////////////////////
+void FTidPacketTransport::DebugEnd()
+{
+#if UE_TRACE_ANALYSIS_DEBUG
+	Threads.RemoveAll([] (const FThreadStream& Thread)
+	{
+		return Thread.Buffer.IsEmpty();
+	});
+
+	DebugUpdate();
+
+	constexpr double MiB = 1024.0 * 1024.0;
+
+	UE_TRACE_ANALYSIS_DEBUG_LOG("");
+	UE_TRACE_ANALYSIS_DEBUG_LOG("TotalPacketSize: %llu bytes (%.1f MiB)", TotalPacketSize, (double)TotalPacketSize / MiB);
+	const uint64 AdjustedTotalDecodedSize = TotalPacketHeaderSize + TotalDecodedSize;
+	UE_TRACE_ANALYSIS_DEBUG_LOG("TotalDecodedSize: %llu bytes + %llu bytes = %llu bytes (%.1f MiB)", TotalPacketHeaderSize, TotalDecodedSize, AdjustedTotalDecodedSize, (double)AdjustedTotalDecodedSize / MiB);
+	const int64 Saving = (int64)AdjustedTotalDecodedSize - (int64)TotalPacketSize;
+	const double SavingPercent = ((double)TotalPacketSize * 100.0) / (double)AdjustedTotalDecodedSize - 100.0;
+	UE_TRACE_ANALYSIS_DEBUG_LOG("Compression Savings: %lli bytes (%.1f MiB; %.2f%%)", Saving, (double)Saving / MiB, SavingPercent);
+
+	UE_TRACE_ANALYSIS_DEBUG_LOG("NumPackets: %u", NumPackets);
+
+	UE_TRACE_ANALYSIS_DEBUG_LOG("MaxBufferSize: %u bytes (%.1f MiB; for thread %u)", MaxBufferSize, (double)MaxBufferSize / MiB, MaxBufferSizeThreadId);
+	const double MaxDataSizePerBufferPercent = ((double)MaxDataSizePerBuffer * 100.0) / (double)TotalDecodedSize;
+	UE_TRACE_ANALYSIS_DEBUG_LOG("MaxDataSizePerBuffer: %u bytes (%.1f%%; for thread %u)", MaxDataSizePerBuffer, MaxDataSizePerBufferPercent, MaxDataSizePerBufferThreadId);
+
+	uint64 MaxDataSizePerThread = 0;
+	uint32 MaxDataSizePerThreadId = 0;
+	for (auto& KV : DataSizePerThread)
+	{
+		if (KV.Value > MaxDataSizePerThread)
+		{
+			MaxDataSizePerThread = KV.Value;
+			MaxDataSizePerThreadId = KV.Key;
+		}
+	}
+	const double MaxDataSizePerThreadPercent = ((double)MaxDataSizePerThread * 100.0) / (double)TotalDecodedSize;
+	UE_TRACE_ANALYSIS_DEBUG_LOG("MaxDataSizePerThread: %llu bytes (%.1f MiB; %.1f%%; for thread %u)", MaxDataSizePerThread, (double)MaxDataSizePerThread / MiB, MaxDataSizePerThreadPercent, MaxDataSizePerThreadId);
+
+	const int32 ThreadCount = Threads.Num();
+	UE_TRACE_ANALYSIS_DEBUG_LOG("Remaining Streaming Threads: %d", ThreadCount);
+	uint64 TotalBufferSize = 0;
+	uint64 UnprocessedDataSize = 0;
+	for (int32 ThreadIndex = 0; ThreadIndex < ThreadCount; ++ThreadIndex)
+	{
+		const FThreadStream& Thread = Threads[ThreadIndex];
+		const uint32 BufferSize = Thread.Buffer.GetBufferSize();
+		const uint32 DataSize = Thread.Buffer.GetRemaining();
+		const double DataSizePercent = ((double)DataSize * 100.0) / (double)TotalDecodedSize;
+		UE_TRACE_ANALYSIS_DEBUG_LOG("[THREAD %d] Tid=%u BufferSize: %u bytes, DataSize: %u bytes (%.1f%%)", ThreadIndex, Thread.ThreadId, BufferSize, DataSize, DataSizePercent);
+		TotalBufferSize += BufferSize;
+		UnprocessedDataSize += DataSize;
+	}
+	UE_TRACE_ANALYSIS_DEBUG_LOG("TotalBufferSize: %llu bytes (%.1f MiB)", TotalBufferSize, (double)TotalBufferSize / MiB);
+
+	const int64 ProcessedDataSize = (int64)TotalDecodedSize - (int64)UnprocessedDataSize;
+	const double ProcessedPercent = ((double)ProcessedDataSize * 100.0) / (double)TotalDecodedSize;
+
+	UE_TRACE_ANALYSIS_DEBUG_LOG("TotalDataSize (processed): %llu bytes (%.1f MiB; %.1f%%)", ProcessedDataSize, (double)ProcessedDataSize / MiB, ProcessedPercent);
+	UE_TRACE_ANALYSIS_DEBUG_LOG("TotalDataSize (unprocessed): %llu bytes (%.1f MiB; %.1f%%)", UnprocessedDataSize, (double)UnprocessedDataSize / MiB, 100.0 - ProcessedPercent);
+
+	const bool bIsEmpty = IsEmpty();
+	check((UnprocessedDataSize == 0 && bIsEmpty) || (UnprocessedDataSize != 0 && !bIsEmpty));
+
+	if (!bIsEmpty)
+	{
+		UE_TRACE_ANALYSIS_DEBUG_LOG("Error: FTidPacketTransport is not empty!");
+	}
+
+	UE_TRACE_ANALYSIS_DEBUG_LOG("FTidPacketTransport::DebugEnd()");
+#endif // UE_TRACE_ANALYSIS_DEBUG
+}
+
 ////////////////////////////////////////////////////////////////////////////////
 uint32 FTidPacketTransport::GetThreadCount() const
 {
diff --git a/Engine/Source/Developer/TraceAnalysis/Private/Analysis/Transport/TidPacketTransport.h b/Engine/Source/Developer/TraceAnalysis/Private/Analysis/Transport/TidPacketTransport.h
--- a/Engine/Source/Developer/TraceAnalysis/Private/Analysis/Transport/TidPacketTransport.h
+++ b/Engine/Source/Developer/TraceAnalysis/Private/Analysis/Transport/TidPacketTransport.h
@@ -5,6 +5,8 @@
 #include "Transport.h"
 #include "Analysis/StreamReader.h"
 #include "Containers/Array.h"
+#include "Containers/Map.h"
+#include "TraceAnalysisDebug.h"
 #include "Trace/Detail/Transport.h"
 
 namespace UE {
@@ -25,6 +27,10 @@ public:
 		Error
 	};
 
+	virtual bool			IsEmpty() const override;
+	virtual void			DebugBegin() override;
+	virtual void			DebugEnd() override;
+
 	ETransportResult		Update();
 	uint32					GetThreadCount() const;
 	FStreamReader*			GetThreadStream(uint32 Index);
@@ -48,6 +54,8 @@ private:
 
 	EReadPacketResult		ReadPacket();
 	FThreadStream*			FindOrAddThread(uint32 ThreadId, bool bAddIfNotFound);
+	void					DebugUpdate();
+
 	TArray<FThreadStream>	Threads = {
 								{ {}, ETransportTid::Events },
 								{ {}, ETransportTid::Importants },
@@ -55,6 +63,19 @@ private:
 
 protected:
 	uint32					Synced = 0x7fff'ffff;
+
+#if UE_TRACE_ANALYSIS_DEBUG
+private:
+	uint64					TotalPacketHeaderSize = 0;
+	uint64					TotalPacketSize = 0;
+	uint64					TotalDecodedSize = 0;
+	uint32					NumPackets = 0;
+	uint32					MaxBufferSize = 0;
+	uint32					MaxBufferSizeThreadId = 0;
+	uint32					MaxDataSizePerBuffer = 0;
+	uint32					MaxDataSizePerBufferThreadId = 0;
+	TMap<uint32, uint64>	DataSizePerThread;
+#endif // UE_TRACE_ANALYSIS_DEBUG
 };
 
 ////////////////////////////////////////////////////////////////////////////////
diff --git a/Engine/Source/Developer/TraceAnalysis/Private/Analysis/Transport/Transport.h b/Engine/Source/Developer/TraceAnalysis/Private/Analysis/Transport/Transport.h
--- a/Engine/Source/Developer/TraceAnalysis/Private/Analysis/Transport/Transport.h
+++ b/Engine/Source/Developer/TraceAnalysis/Private/Analysis/Transport/Transport.h
@@ -18,9 +18,13 @@ public:
 	template <typename RetType>
 	RetType const*			GetPointer(uint32 BlockSize);
 	virtual void			Advance(uint32 BlockSize);
+	virtual bool			IsEmpty() const { return true; }
+	virtual void			DebugBegin() {}
+	virtual void			DebugEnd() {}
 
 protected:
 	virtual const uint8*	GetPointerImpl(uint32 BlockSize);
+
 	FStreamReader*			Reader;
 };
 
diff --git a/Engine/Source/Developer/TraceAnalysis/Private/TraceAnalysisDebug.cpp b/Engine/Source/Developer/TraceAnalysis/Private/TraceAnalysisDebug.cpp
new file mode 100644
--- /dev/null
+++ b/Engine/Source/Developer/TraceAnalysis/Private/TraceAnalysisDebug.cpp
@@ -0,0 +1,3 @@
+// Copyright Epic Games, Inc. All Rights Reserved.
+
+#include "TraceAnalysisDebug.h"
diff --git a/Engine/Source/Developer/TraceAnalysis/Public/Trace/Analyzer.h b/Engine/Source/Developer/TraceAnalysis/Public/Trace/Analyzer.h
--- a/Engine/Source/Developer/TraceAnalysis/Public/Trace/Analyzer.h
+++ b/Engine/Source/Developer/TraceAnalysis/Public/Trace/Analyzer.h
@@ -5,6 +5,7 @@
 #include "CoreTypes.h"
 #include "Containers/ArrayView.h"
 #include "Containers/StringView.h"
+#include "TraceAnalysisDebug.h"
 #include "Trace/Trace.h"
 
 #include <type_traits>
@@ -66,9 +67,13 @@ public:
 		/** Is this field signed (only relevant for integer types) */
 		bool IsSigned() const;
 
+#if UE_TRACE_ANALYSIS_DEBUG_API
+		/** Offset from the start of the event to this field's data. */
+		uint32 GetOffset() const;
+#endif // UE_TRACE_ANALYSIS_DEBUG_API
+
 		/** Gets the size in bytes for this field */
 		uint8 GetSize() const;
-
 	};
 
 	struct FEventFieldHandle
@@ -83,6 +88,11 @@ public:
 		 * guaranteed to be the same for the same event from one trace to the next. */
 		uint32 GetId() const;
 
+#if UE_TRACE_ANALYSIS_DEBUG_API
+		/** Returns the event's flags. */
+		uint8 GetFlags() const;
+#endif // UE_TRACE_ANALYSIS_DEBUG_API
+
 		/** The name of the event. */
 		const ANSICHAR* GetName() const;
 
@@ -98,7 +108,10 @@ public:
 		/** By-index access to fields' type information. */
 		const FEventFieldInfo* GetFieldInfo(uint32 Index) const;
 
-		/** Returns a handle that can used to access events' fields. There is 
+		/** Returns the field index or -1 (if the event does not contains a field with the specified name). */
+		int32 GetFieldIndex(const ANSICHAR* FieldName) const;
+
+		/** Returns a handle that can used to access events' fields. There is
 		 * loose validation via ValueType, but one should still exercise caution
 		 * when reading fields with handles.
 		 * @param ValueType The intended type that the field will be interpreted as */
@@ -117,9 +130,20 @@ public:
 
 	struct TRACEANALYSIS_API FArrayReader
 	{
-		/* Returns the number of elements in the array */
+		/* Returns the number of elements in the array. */
 		uint32 Num() const;
 
+#if UE_TRACE_ANALYSIS_DEBUG_API
+		/* Returns the pointer to the raw data array. */
+		const uint8* GetRawData() const;
+
+		/* Returns the size in bytes of the raw data array. */
+		uint32 GetRawDataSize() const;
+
+		/* Returns the size and type of an array element. */
+		int8 GetSizeAndType() const;
+#endif // UE_TRACE_ANALYSIS_DEBUG_API
+
 	protected:
 		const void* GetImpl(uint32 Index, int8& SizeAndType) const;
 	};
@@ -136,6 +160,15 @@ public:
 		const ValueType* GetData() const;
 	};
 
+	enum class EStyle : uint32
+	{
+		Normal,
+		EnterScope,
+		LeaveScope,
+	};
+
+	struct FOnEventContext;
+
 	struct TRACEANALYSIS_API FEventData
 	{
 		/** Returns an object describing the underlying event's type. */
@@ -179,7 +212,7 @@ public:
 
 		template<typename DefinitionType>
 		TEventRef<DefinitionType> GetReferenceValue(uint32 FieldIndex) const;
-		
+
 		/** If this is a spec event, gets the unique Id for this spec.
 		 * @return A valid spec id if the event is valid, otherwise an empty id.
 		 */
@@ -189,7 +222,7 @@ public:
 		uint32 GetSize() const;
 
 		/** Serializes the event to Cbor object.
-		 * @param Recipient of the Cbor serialization. Data is appeneded to Out. */
+		 * @param Recipient of the Cbor serialization. Data is appended to Out. */
 		void SerializeToCbor(TArray<uint8>& Out) const;
 
 		/**
@@ -207,9 +240,22 @@ public:
 		/** Returns the size of the events attachment, or 0 if none. */
 		uint32 GetAttachmentSize() const;
 
+#if UE_TRACE_ANALYSIS_DEBUG_API
+		/** Provides a pointer to the raw event data. */
+		const uint8* GetRawPointer() const;
+
+		/** Returns the size of the raw event data (including attachment). */
+		uint32 GetRawSize() const;
+
+		/** Returns the total uncompressed size of the aux data (including size of aux headers and terminator), in bytes. */
+		uint32 GetAuxSize() const;
+
+		/** Returns the total uncompressed size of the event, in bytes (including headers and aux data). */
+		uint32 GetTotalSize(IAnalyzer::EStyle Style, const IAnalyzer::FOnEventContext& Context, uint32 ProtocolVersion = 7) const;
+#endif // UE_TRACE_ANALYSIS_DEBUG_API
+
 	private:
 		bool IsDefinitionImpl(uint32& OutTypeId) const;
-		uint32 GetReferenceTypeIdImpl(FEventFieldHandle Handle) const;
 		const void* GetReferenceValueImpl(const char* FieldName, uint16& OutSizeType, uint32& OutTypeUid) const;
 		const void* GetReferenceValueImpl(uint32 FieldIndex, uint32& OutTypeUid) const;
 		const void* GetValueImpl(const ANSICHAR* FieldName, int8& SizeAndType) const;
@@ -221,7 +267,7 @@ public:
 		/* Returns the trace-specific id for the thread */
 		uint32 GetId() const;
 
-		/* Returns the system if for the thread. Because this may not be known by
+		/* Returns the system id for the thread. Because this may not be known by
 		 * trace and because IDs can be reused by the system, relying on the value
 		 * of this is discouraged. */
 		uint32 GetSystemId() const;
@@ -241,7 +287,7 @@ public:
 		/** Returns the integer timestamp for the event or zero if there no associated timestamp. */
 		uint64 GetTimestamp() const;
 
-		/** Time of the event in seconds (from teh start of the trace). Zero if there is no time for the event. */
+		/** Time of the event in seconds (from the start of the trace). Zero if there is no time for the event. */
 		double AsSeconds() const;
 
 		/** Returns a timestamp for the event compatible with FPlatformTime::Cycle64(), or zero if the event has no timestamp. */
@@ -294,13 +340,6 @@ public:
 		return true;
 	}
 
-	enum class EStyle : uint32
-	{
-		Normal,
-		EnterScope,
-		LeaveScope,
-	};
-
 	/** For each event subscribed to in OnAnalysisBegin(), the analysis engine
 	 * will call this method when those events are encountered in a trace log
 	 * @param RouteId User-provided identifier given when subscribing to a particular event.
@@ -312,6 +351,12 @@ public:
 		return true;
 	}
 
+#if UE_TRACE_ANALYSIS_DEBUG_API
+	virtual void OnVersion(uint32 TransportVersion, uint32 ProtocolVersion)
+	{
+	}
+#endif // UE_TRACE_ANALYSIS_DEBUG_API
+
 private:
 	template <typename ValueType> static ValueType CoerceValue(const void* Addr, int8 SizeAndType);
 };
@@ -451,6 +496,6 @@ TEventRef<DefinitionType> IAnalyzer::FEventData::GetReferenceValue(uint32 FieldI
 	DefinitionType* Id = (DefinitionType*) GetReferenceValueImpl(FieldIndex, RefTypeUid);
 	return MakeEventRef<DefinitionType>(*Id, RefTypeUid);
 }
-	
+
 } // namespace Trace
 } // namespace UE
diff --git a/Engine/Source/Developer/TraceAnalysis/Public/TraceAnalysisDebug.h b/Engine/Source/Developer/TraceAnalysis/Public/TraceAnalysisDebug.h
new file mode 100644
--- /dev/null
+++ b/Engine/Source/Developer/TraceAnalysis/Public/TraceAnalysisDebug.h
@@ -0,0 +1,96 @@
+// Copyright Epic Games, Inc. All Rights Reserved.
+
+#pragma once
+
+// Enables debug API in UE::TRACE::IAnalyzer class (and sub-types).
+// Enabling it should not add any impact on analysis' performance or behaviour.
+#if !defined(UE_TRACE_ANALYSIS_DEBUG_API)
+	#define UE_TRACE_ANALYSIS_DEBUG_API 0
+#endif
+
+// Enables debug functionality and logging in TraceAnalysis code.
+// Enabling it can add significant impact on analysis' performance and it could alter behaviour.
+#if !defined(UE_TRACE_ANALYSIS_DEBUG)
+	#define UE_TRACE_ANALYSIS_DEBUG 0
+#endif
+
+// Default low level logging API used when UE_TRACE_ANALYSIS_DEBUG is enabled.
+// For now, only ANSICHAR* strings are supported for logging.
+#if !defined(UE_TRACE_ANALYSIS_DEBUG_LOG)
+#if UE_TRACE_ANALYSIS_DEBUG
+
+// Default implementation of UE_TRACE_ANALYSIS_DEBUG_LOG API:
+// 1: only printf
+// 2: printf with TAnsiStringBuilder
+// 3: GLog
+// 4: FPlatformMisc::LowLevelOutputDebugStringf
+#if !defined(UE_TRACE_ANALYSIS_DEBUG_LOG_IMPL)
+	#define UE_TRACE_ANALYSIS_DEBUG_LOG_IMPL 1
+#endif
+
+#if UE_TRACE_ANALYSIS_DEBUG_LOG_IMPL == 1 // only printf
+
+	#define UE_TRACE_ANALYSIS_DEBUG_LOG(format, ...)        printf(format "\n", __VA_ARGS__)
+	#define UE_TRACE_ANALYSIS_DEBUG_BeginStringBuilder()
+	#define UE_TRACE_ANALYSIS_DEBUG_Append(s)               printf("%s", s)
+	#define UE_TRACE_ANALYSIS_DEBUG_Appendf(format, ...)    printf(format, __VA_ARGS__)
+	#define UE_TRACE_ANALYSIS_DEBUG_AppendChar(c)           printf("%c", c)
+	#define UE_TRACE_ANALYSIS_DEBUG_EndStringBuilder()      printf("\n")
+	#define UE_TRACE_ANALYSIS_DEBUG_ResetStringBuilder()
+
+#elif UE_TRACE_ANALYSIS_DEBUG_LOG_IMPL == 2 // printf with TAnsiStringBuilder
+
+	#define UE_TRACE_ANALYSIS_DEBUG_LOG(format, ...)        printf(format "\n", __VA_ARGS__)
+	#define UE_TRACE_ANALYSIS_DEBUG_BeginStringBuilder()    TAnsiStringBuilder<1024> StringBuilder;
+	#define UE_TRACE_ANALYSIS_DEBUG_Append                  StringBuilder.Append
+	#define UE_TRACE_ANALYSIS_DEBUG_Appendf                 StringBuilder.Appendf
+	#define UE_TRACE_ANALYSIS_DEBUG_AppendChar              StringBuilder.AppendChar
+	#define UE_TRACE_ANALYSIS_DEBUG_EndStringBuilder()      printf("%s\n", StringBuilder.ToString())
+	#define UE_TRACE_ANALYSIS_DEBUG_ResetStringBuilder()    StringBuilder.Reset()
+
+#elif UE_TRACE_ANALYSIS_DEBUG_LOG_IMPL == 3 // GLog
+
+	#include "Containers/StringConv.h"
+	#include "CoreGlobals.h"
+	#include "Misc/OutputDeviceRedirector.h"
+	#include "Misc/StringBuilder.h"
+
+	#define UE_TRACE_ANALYSIS_DEBUG_LOG(format, ...)        { TAnsiStringBuilder<1024> SB; SB.Appendf(format, __VA_ARGS__); GLog->Logf(TEXT("%hs"), SB.ToString()); }
+	#define UE_TRACE_ANALYSIS_DEBUG_BeginStringBuilder()    TAnsiStringBuilder<1024> StringBuilder;
+	#define UE_TRACE_ANALYSIS_DEBUG_Append                  StringBuilder.Append
+	#define UE_TRACE_ANALYSIS_DEBUG_Appendf                 StringBuilder.Appendf
+	#define UE_TRACE_ANALYSIS_DEBUG_AppendChar              StringBuilder.AppendChar
+	#define UE_TRACE_ANALYSIS_DEBUG_EndStringBuilder()      GLog->Logf(TEXT("%hs"), StringBuilder.ToString())
+	#define UE_TRACE_ANALYSIS_DEBUG_ResetStringBuilder()    StringBuilder.Reset()
+
+#elif UE_TRACE_ANALYSIS_DEBUG_LOG_IMPL == 4 // FPlatformMisc::LowLevelOutputDebugStringf
+
+	#include "Misc/StringBuilder.h"
+
+	#define UE_TRACE_ANALYSIS_DEBUG_LOG(format, ...)        { TAnsiStringBuilder<1024> SB; SB.Appendf(format, __VA_ARGS__); FPlatformMisc::LowLevelOutputDebugStringf(TEXT("%hs\n"), SB.ToString()); }
+	#define UE_TRACE_ANALYSIS_DEBUG_BeginStringBuilder()
+	#define UE_TRACE_ANALYSIS_DEBUG_Append(s)               FPlatformMisc::LowLevelOutputDebugStringf(TEXT("%hs"), s)
+	#define UE_TRACE_ANALYSIS_DEBUG_Appendf(format, ...)    { TAnsiStringBuilder<1024> SB; SB.Appendf(format, __VA_ARGS__); FPlatformMisc::LowLevelOutputDebugStringf(TEXT("%hs"), SB.ToString()); }
+	#define UE_TRACE_ANALYSIS_DEBUG_AppendChar(c)           FPlatformMisc::LowLevelOutputDebugStringf(TEXT("%hc"), c)
+	#define UE_TRACE_ANALYSIS_DEBUG_EndStringBuilder()      FPlatformMisc::LowLevelOutputDebugString(TEXT("\n"))
+	#define UE_TRACE_ANALYSIS_DEBUG_ResetStringBuilder()
+
+#endif // UE_TRACE_ANALYSIS_DEBUG_LOG_IMPL
+
+#else // UE_TRACE_ANALYSIS_DEBUG
+
+	#define UE_TRACE_ANALYSIS_DEBUG_LOG(format, ...)
+	#define UE_TRACE_ANALYSIS_DEBUG_BeginStringBuilder()
+	#define UE_TRACE_ANALYSIS_DEBUG_Append(s)
+	#define UE_TRACE_ANALYSIS_DEBUG_Appendf(format, ...)
+	#define UE_TRACE_ANALYSIS_DEBUG_AppendChar(c)
+	#define UE_TRACE_ANALYSIS_DEBUG_EndStringBuilder()
+	#define UE_TRACE_ANALYSIS_DEBUG_ResetStringBuilder()
+
+#endif  // UE_TRACE_ANALYSIS_DEBUG
+#endif // !defined(UE_TRACE_ANALYSIS_DEBUG_LOG)
+
+// Log level (1 = minimum, 2 = normal, 3 = verbose, 4 = verbose+). It is used only when UE_TRACE_ANALYSIS_DEBUG is enabled.
+#if !defined(UE_TRACE_ANALYSIS_DEBUG_LEVEL)
+	#define UE_TRACE_ANALYSIS_DEBUG_LEVEL 2
+#endif
diff --git a/Engine/Source/Runtime/Cbor/Private/CborReader.cpp b/Engine/Source/Runtime/Cbor/Private/CborReader.cpp
--- a/Engine/Source/Runtime/Cbor/Private/CborReader.cpp
+++ b/Engine/Source/Runtime/Cbor/Private/CborReader.cpp
@@ -6,6 +6,7 @@ FCborReader::FCborReader(FArchive* InStream, ECborEndianness InReaderEndianness)
 	: Stream(InStream)
 	, Endianness(InReaderEndianness)
 {
+	check(InStream != nullptr);
 	ContextStack.Emplace();
 }
 
@@ -69,6 +70,7 @@ bool FCborReader::ReadNext(FCborContext& OutContext)
 		// Report parent context container type
 		OutContext.RawTextValue.Add((char)ParentContext.MajorType());
 		// Done with parent context
+		check(ContextStack.Num() > 1);
 		ContextStack.Pop();
 		return true;
 	}
@@ -104,6 +106,7 @@ bool FCborReader::ReadNext(FCborContext& OutContext)
 		// Report parent context container type
 		OutContext.RawTextValue.Add((char)ParentContext.MajorType());
 		// Done with parent context
+		check(ContextStack.Num() > 1);
 		ContextStack.Pop();
 		return true;
 	}
@@ -185,6 +188,12 @@ bool FCborReader::ReadNext(FCborContext& OutContext)
 
 bool FCborReader::SkipContainer(ECborCode ContainerType)
 {
+	// Invalid stream error
+	if (Stream == nullptr)
+	{
+		return false;
+	}
+
 	if (GetContext().MajorType() != ContainerType)
 	{
 		return false;
