Add store client functionality to set sponsor mode

diff --git a/Engine/Source/Developer/TraceAnalysis/Private/Store/StoreClient.cpp b/Engine/Source/Developer/TraceAnalysis/Private/Store/StoreClient.cpp
--- a/Engine/Source/Developer/TraceAnalysis/Private/Store/StoreClient.cpp
+++ b/Engine/Source/Developer/TraceAnalysis/Private/Store/StoreClient.cpp
@@ -119,6 +119,7 @@ public:
 	bool					GetSessionInfoByTraceGuid(const FGuid& TraceGuid);
 							bool SetStoreDirectories(const TCHAR* StoreDir, const TArray<FString>& AddWatchDirs,
 													const TArray<FString>& RemoveWatchDirs);
+							bool SetSponsored(bool bSponsored);
 
 private:
 	bool					Communicate(const FPayload& Payload);
@@ -422,6 +423,14 @@ bool FStoreCborClient::SetStoreDirectories(const TCHAR* StoreDir, const TArray<F
 	return Communicate(Builder.Done());
 }
 	
+////////////////////////////////////////////////////////////////////////////////
+bool FStoreCborClient::SetSponsored(bool bSponsored)
+{
+	TPayloadBuilder<> Builder("v1/settings/write");
+	Builder.AddInteger("Sponsored", bSponsored ? 1 : 0);
+	return Communicate(Builder.Done());
+}
+	
 } // namespace Trace
 } // namespace UE
 
@@ -443,6 +452,20 @@ uint32 FStoreClient::FStatus::GetRecorderPort() const
 	return Response->GetUint32Checked("recorder_port", 0);
 }
 
+////////////////////////////////////////////////////////////////////////////////
+uint32 FStoreClient::FStatus::GetStorePort() const
+{
+	const auto* Response = (const FResponse*)this;
+	return Response->GetUint32Checked("store_port", 0);
+}
+
+////////////////////////////////////////////////////////////////////////////////
+bool FStoreClient::FStatus::GetSponsored() const
+{
+	const auto* Response = (const FResponse*)this;
+	return !!Response->GetInteger("sponsored", 0);
+}
+
 ////////////////////////////////////////////////////////////////////////////////
 uint32 FStoreClient::FStatus::GetChangeSerial() const
 {
@@ -683,6 +706,13 @@ bool FStoreClient::SetStoreDirectories(const TCHAR* StoreDir,
 	return Self->SetStoreDirectories(StoreDir, AddWatchDirs, RemoveWatchDirs);
 }
 
+////////////////////////////////////////////////////////////////////////////////
+bool FStoreClient::SetSponsored(bool bSponsored)
+{
+	auto* Self = (FStoreCborClient*) this;
+	return Self->SetSponsored(bSponsored);
+}
+	
 ////////////////////////////////////////////////////////////////////////////////
 uint32 FStoreClient::GetSessionCount() const
 {
diff --git a/Engine/Source/Developer/TraceAnalysis/Public/Trace/StoreClient.h b/Engine/Source/Developer/TraceAnalysis/Public/Trace/StoreClient.h
--- a/Engine/Source/Developer/TraceAnalysis/Public/Trace/StoreClient.h
+++ b/Engine/Source/Developer/TraceAnalysis/Public/Trace/StoreClient.h
@@ -21,6 +21,8 @@ public:
 	{
 		FUtf8StringView	GetStoreDir() const;
 		uint32			GetRecorderPort() const;
+		uint32			GetStorePort() const;
+		bool			GetSponsored() const;
 		uint32			GetChangeSerial() const;
 		uint32			GetSettingsSerial() const;
 		void			GetWatchDirectories(TArray<FString>& OutDirs) const;
@@ -65,6 +67,7 @@ public:
 	FTraceData			ReadTrace(uint32 Id);
 	bool SetStoreDirectories(const TCHAR* StoreDir, const TArray<FString>& AddWatchDirs,
 							const TArray<FString>& RemoveWatchDirs);
+	bool				SetSponsored(bool bSponsored);
 #if 0
 	template <typename Lambda> uint32 GetTraceInfos(uint32 StartIndex, uint32 Count, Lambda&& Callback) const;
 #endif // 0
