[Insights] Updated limits for detecting too much data accumulated by trace analysis (number of parsed events).

diff --git a/Engine/Source/Developer/TraceAnalysis/Private/Analysis/Engine.cpp b/Engine/Source/Developer/TraceAnalysis/Private/Analysis/Engine.cpp
--- a/Engine/Source/Developer/TraceAnalysis/Private/Analysis/Engine.cpp
+++ b/Engine/Source/Developer/TraceAnalysis/Private/Analysis/Engine.cpp
@@ -3939,10 +3939,13 @@ FProtocol5Stage::EStatus FProtocol5Stage::OnDataNormal(const FMachineContext& Co
 	// Test if analysis has accumulated too much data (parsed events not dispatched yet).
 	// But, only enforce the limit after we have received at least one SYNC package
 	// (e.g. server traces can accumulate large amounts of data before first SYNC event).
-	constexpr int32 MaxAvailableParsedEvents = 10'000'000;
+	constexpr int32 MaxAvailableEventsHighLimit = 90'000'000;
+	constexpr int32 MaxAvailableEventsLowLimit = 50'000'000;
 	constexpr uint32 MaxNextSerialWaitCount = 20;
+	bool bSkipSerialNow = false;
 	if (SyncCount > 0 &&
-		NumAvailableEvents > MaxAvailableParsedEvents &&
+		!bSkipSerialError &&
+		NumAvailableEvents > MaxAvailableEventsHighLimit &&
 		NextSerialWaitCount > MaxNextSerialWaitCount)
 	{
 		UE_TRACE_ANALYSIS_DEBUG_LOG("Error: Trace analysis accumulated too much data (%d parsed events) and will start to skip the missing serial sync events!", NumAvailableEvents);
@@ -3953,19 +3956,28 @@ FProtocol5Stage::EStatus FProtocol5Stage::OnDataNormal(const FMachineContext& Co
 				EAnalysisMessageSeverity::Error,
 				TEXT("Trace analysis accumulated too much data (%d parsed events) and will start to skip the missing serial sync events!"), NumAvailableEvents);
 		}
-		// Skip serials and continue to dispatch parsed events.
-		bSkipSerial = true;
-		NextSerialWaitCount = 0;
-		int32 NumDispatchedEvents = DispatchNormalEvents(Context, EventDescHeap);
-		if (NumDispatchedEvents < 0)
+		bSkipSerialNow = true;
+	}
+	if (bSkipSerialNow ||
+		(bSkipSerialError && NumAvailableEvents > MaxAvailableEventsLowLimit))
+	{
+		do
 		{
-			return EStatus::Error;
-		}
+			// Skip serials and continue to dispatch parsed events.
+			bSkipSerial = true;
+			NextSerialWaitCount = 0;
+			int32 NumDispatchedEvents = DispatchNormalEvents(Context, EventDescHeap);
+			if (NumDispatchedEvents < 0)
+			{
+				return EStatus::Error;
+			}
 #if UE_TRACE_ANALYSIS_DEBUG && UE_TRACE_ANALYSIS_DEBUG_LEVEL >= 2
-		UE_TRACE_ANALYSIS_DEBUG_LOG("Skipped serials and dispatched %d normal events (%d --> %d)", NumDispatchedEvents, NumAvailableEvents, NumAvailableEvents - NumDispatchedEvents);
+			UE_TRACE_ANALYSIS_DEBUG_LOG("Skipped serials and dispatched %d normal events (%d --> %d)", NumDispatchedEvents, NumAvailableEvents, NumAvailableEvents - NumDispatchedEvents);
 #endif
-		NumAvailableEvents -= NumDispatchedEvents;
-		check(NumAvailableEvents >= 0);
+			NumAvailableEvents -= NumDispatchedEvents;
+			check(NumAvailableEvents >= 0);
+		}
+		while (NumAvailableEvents > MaxAvailableEventsLowLimit);
 	}
 
 	// If there are any streams left in the heap then we are unable to proceed
