[Insights] Fixed constness for FStoreClient.

diff --git a/Engine/Source/Developer/TraceAnalysis/Private/Store/StoreClient.cpp b/Engine/Source/Developer/TraceAnalysis/Private/Store/StoreClient.cpp
--- a/Engine/Source/Developer/TraceAnalysis/Private/Store/StoreClient.cpp
+++ b/Engine/Source/Developer/TraceAnalysis/Private/Store/StoreClient.cpp
@@ -106,28 +106,28 @@ public:
 	uint32					GetStorePort() const;
 	const FResponse&		GetResponse() const;
 	bool					Connect(const TCHAR* Host, uint16 Port);
-	bool					GetStatus();
-	bool					GetVersion();
-	bool					GetTraceCount();
-	bool					GetTraceInfo(uint32 Index);
-	bool					GetTraceInfoById(uint32 Id);
-	FTraceDataStream*		ReadTrace(uint32 Id);
-	bool					GetSessionCount();
-	bool					GetSessionInfo(uint32 Index);
-	bool					GetSessionInfoById(uint32 Id);
-	bool					GetSessionInfoByTraceId(uint32 TraceId);
-	bool					GetSessionInfoByTraceGuid(const FGuid& TraceGuid);
-							bool SetStoreDirectories(const TCHAR* StoreDir, const TArray<FString>& AddWatchDirs,
-													const TArray<FString>& RemoveWatchDirs);
-							bool SetSponsored(bool bSponsored);
+	bool					GetStatus() const;
+	bool					GetVersion() const;
+	bool					GetTraceCount() const;
+	bool					GetTraceInfo(uint32 Index) const;
+	bool					GetTraceInfoById(uint32 Id) const;
+	FTraceDataStream*		ReadTrace(uint32 Id) const;
+	bool					GetSessionCount() const;
+	bool					GetSessionInfo(uint32 Index) const;
+	bool					GetSessionInfoById(uint32 Id) const;
+	bool					GetSessionInfoByTraceId(uint32 TraceId) const;
+	bool					GetSessionInfoByTraceGuid(const FGuid& TraceGuid) const;
+	bool					SetStoreDirectories(const TCHAR* StoreDir, const TArray<FString>& AddWatchDirs, const TArray<FString>& RemoveWatchDirs);
+	bool					SetSponsored(bool bSponsored);
 
 private:
-	bool					Communicate(const FPayload& Payload);
-	asio::io_context&		IoContext;
-	asio::ip::tcp::socket	Socket;
-	FResponse				Response;
-	FString					Host;
-	uint16					Port;
+	bool					Communicate(const FPayload& Payload) const;
+
+	asio::io_context&				IoContext;
+	mutable asio::ip::tcp::socket	Socket;
+	mutable FResponse				Response;
+	FString							Host;
+	uint16							Port;
 };
 
 ////////////////////////////////////////////////////////////////////////////////
@@ -224,7 +224,7 @@ bool FStoreCborClient::Connect(const TCHAR* InHost, uint16 InPort)
 }
 
 ////////////////////////////////////////////////////////////////////////////////
-bool FStoreCborClient::Communicate(const FPayload& Payload)
+bool FStoreCborClient::Communicate(const FPayload& Payload) const
 {
 	if (!Socket.is_open())
 	{
@@ -275,7 +275,7 @@ bool FStoreCborClient::Communicate(const FPayload& Payload)
 }
 
 ////////////////////////////////////////////////////////////////////////////////
-bool FStoreCborClient::GetStatus()
+bool FStoreCborClient::GetStatus() const
 {
 	TPayloadBuilder<32> Builder("v1/status");
 	FPayload Payload = Builder.Done();
@@ -283,15 +283,15 @@ bool FStoreCborClient::GetStatus()
 }
 
 ////////////////////////////////////////////////////////////////////////////////
-bool FStoreCborClient::GetVersion()
+bool FStoreCborClient::GetVersion() const
 {
 	TPayloadBuilder<32> Builder("v1/version");
 	FPayload Payload = Builder.Done();
 	return Communicate(Payload);
 }
-	
+
 ////////////////////////////////////////////////////////////////////////////////
-bool FStoreCborClient::GetTraceCount()
+bool FStoreCborClient::GetTraceCount() const
 {
 	TPayloadBuilder<32> Builder("v1/trace/count");
 	FPayload Payload = Builder.Done();
@@ -299,7 +299,7 @@ bool FStoreCborClient::GetTraceCount()
 }
 
 ////////////////////////////////////////////////////////////////////////////////
-bool FStoreCborClient::GetTraceInfo(uint32 Index)
+bool FStoreCborClient::GetTraceInfo(uint32 Index) const
 {
 	TPayloadBuilder<> Builder("v1/trace/info");
 	Builder.AddInteger("index", Index);
@@ -308,7 +308,7 @@ bool FStoreCborClient::GetTraceInfo(uint32 Index)
 }
 
 ////////////////////////////////////////////////////////////////////////////////
-bool FStoreCborClient::GetTraceInfoById(uint32 Id)
+bool FStoreCborClient::GetTraceInfoById(uint32 Id) const
 {
 	TPayloadBuilder<> Builder("v1/trace/info");
 	Builder.AddInteger("id", int32(Id));
@@ -317,7 +317,7 @@ bool FStoreCborClient::GetTraceInfoById(uint32 Id)
 }
 
 ////////////////////////////////////////////////////////////////////////////////
-FTraceDataStream* FStoreCborClient::ReadTrace(uint32 Id)
+FTraceDataStream* FStoreCborClient::ReadTrace(uint32 Id) const
 {
 	TPayloadBuilder<> Builder("v1/trace/read");
 	Builder.AddInteger("id", Id);
@@ -356,7 +356,7 @@ FTraceDataStream* FStoreCborClient::ReadTrace(uint32 Id)
 }
 
 ////////////////////////////////////////////////////////////////////////////////
-bool FStoreCborClient::GetSessionCount()
+bool FStoreCborClient::GetSessionCount() const
 {
 	TPayloadBuilder<> Builder("v1/session/count");
 	FPayload Payload = Builder.Done();
@@ -364,7 +364,7 @@ bool FStoreCborClient::GetSessionCount()
 }
 
 ////////////////////////////////////////////////////////////////////////////////
-bool FStoreCborClient::GetSessionInfo(uint32 Index)
+bool FStoreCborClient::GetSessionInfo(uint32 Index) const
 {
 	TPayloadBuilder<> Builder("v1/session/info");
 	Builder.AddInteger("index", int32(Index));
@@ -373,7 +373,7 @@ bool FStoreCborClient::GetSessionInfo(uint32 Index)
 }
 
 ////////////////////////////////////////////////////////////////////////////////
-bool FStoreCborClient::GetSessionInfoById(uint32 Id)
+bool FStoreCborClient::GetSessionInfoById(uint32 Id) const
 {
 	TPayloadBuilder<> Builder("v1/session/info");
 	Builder.AddInteger("id", int32(Id));
@@ -382,7 +382,7 @@ bool FStoreCborClient::GetSessionInfoById(uint32 Id)
 }
 
 ////////////////////////////////////////////////////////////////////////////////
-bool FStoreCborClient::GetSessionInfoByTraceId(uint32 TraceId)
+bool FStoreCborClient::GetSessionInfoByTraceId(uint32 TraceId) const
 {
 	TPayloadBuilder<> Builder("v1/session/info");
 	Builder.AddInteger("trace_id", int32(TraceId));
@@ -391,7 +391,7 @@ bool FStoreCborClient::GetSessionInfoByTraceId(uint32 TraceId)
 }
 
 ////////////////////////////////////////////////////////////////////////////////
-bool FStoreCborClient::GetSessionInfoByTraceGuid(const FGuid& TraceGuid)
+bool FStoreCborClient::GetSessionInfoByTraceGuid(const FGuid& TraceGuid) const
 {
 	TPayloadBuilder<> Builder("v1/session/info");
 	Builder.AddString("trace_guid", TCHAR_TO_ANSI(*TraceGuid.ToString()));
@@ -422,7 +422,7 @@ bool FStoreCborClient::SetStoreDirectories(const TCHAR* StoreDir, const TArray<F
 	}
 	return Communicate(Builder.Done());
 }
-	
+
 ////////////////////////////////////////////////////////////////////////////////
 bool FStoreCborClient::SetSponsored(bool bSponsored)
 {
@@ -430,151 +430,160 @@ bool FStoreCborClient::SetSponsored(bool bSponsored)
 	Builder.AddInteger("Sponsored", bSponsored ? 1 : 0);
 	return Communicate(Builder.Done());
 }
-	
-} // namespace Trace
-} // namespace UE
 
 
-namespace UE {
-namespace Trace {
 
+////////////////////////////////////////////////////////////////////////////////
+// FStoreClient::FStatus
 ////////////////////////////////////////////////////////////////////////////////
 FUtf8StringView FStoreClient::FStatus::GetStoreDir() const
 {
-	const auto* Response = (const FResponse*)this;
+	const FResponse* Response = (const FResponse*)this;
 	return Response->GetString("store_dir", "");
 }
 
 ////////////////////////////////////////////////////////////////////////////////
 uint32 FStoreClient::FStatus::GetRecorderPort() const
 {
-	const auto* Response = (const FResponse*)this;
+	const FResponse* Response = (const FResponse*)this;
 	return Response->GetUint32Checked("recorder_port", 0);
 }
 
 ////////////////////////////////////////////////////////////////////////////////
 uint32 FStoreClient::FStatus::GetStorePort() const
 {
-	const auto* Response = (const FResponse*)this;
+	const FResponse* Response = (const FResponse*)this;
 	return Response->GetUint32Checked("store_port", 0);
 }
 
 ////////////////////////////////////////////////////////////////////////////////
 bool FStoreClient::FStatus::GetSponsored() const
 {
-	const auto* Response = (const FResponse*)this;
+	const FResponse* Response = (const FResponse*)this;
 	return !!Response->GetInteger("sponsored", 0);
 }
 
 ////////////////////////////////////////////////////////////////////////////////
 uint32 FStoreClient::FStatus::GetChangeSerial() const
 {
-	const auto* Response = (const FResponse*)this;
+	const FResponse* Response = (const FResponse*)this;
 	return Response->GetUint32Checked("change_serial", 0);
 }
 
 ////////////////////////////////////////////////////////////////////////////////
 uint32 FStoreClient::FStatus::GetSettingsSerial() const
 {
-	const auto* Response = (const FResponse*)this;
+	const FResponse* Response = (const FResponse*)this;
 	return Response->GetUint32Checked("settings_serial", 0);
 }
 
 ////////////////////////////////////////////////////////////////////////////////
 void FStoreClient::FStatus::GetWatchDirectories(TArray<FString>& OutDirs) const
 {
-	const auto* Response = (const FResponse*)this;
+	const FResponse* Response = (const FResponse*)this;
 	Response->GetStringArray("watch_dirs", OutDirs);
 }
 
+
+
+////////////////////////////////////////////////////////////////////////////////
+// FStoreClient::FVersion
 ////////////////////////////////////////////////////////////////////////////////
 uint32 FStoreClient::FVersion::GetMajorVersion() const
 {
-	const auto* Response = (const FResponse*)this;
+	const FResponse* Response = (const FResponse*)this;
 	return Response->GetUint32Checked("major", 0);
 }
 
 ////////////////////////////////////////////////////////////////////////////////
 uint32 FStoreClient::FVersion::GetMinorVersion() const
 {
-	const auto* Response = (const FResponse*)this;
+	const FResponse* Response = (const FResponse*)this;
 	return Response->GetUint32Checked("minor", 0);
 }
 
 ////////////////////////////////////////////////////////////////////////////////
 FUtf8StringView FStoreClient::FVersion::GetConfiguration() const
 {
-	const auto* Response = (const FResponse*)this;
+	const FResponse* Response = (const FResponse*)this;
 	return Response->GetString("configuration", "unknown");
 }
 
+
+
+////////////////////////////////////////////////////////////////////////////////
+// FStoreClient::FTraceInfo
 ////////////////////////////////////////////////////////////////////////////////
 uint32 FStoreClient::FTraceInfo::GetId() const
 {
-	const auto* Response = (const FResponse*)this;
+	const FResponse* Response = (const FResponse*)this;
 	return Response->GetUint32Checked("id", 0);
 }
 
 ////////////////////////////////////////////////////////////////////////////////
 uint64 FStoreClient::FTraceInfo::GetSize() const
 {
-	const auto* Response = (const FResponse*)this;
+	const FResponse* Response = (const FResponse*)this;
 	return Response->GetUint64Checked("size", 0);
 }
 
 ////////////////////////////////////////////////////////////////////////////////
 uint64 FStoreClient::FTraceInfo::GetTimestamp() const
 {
-	const auto* Response = (const FResponse*)this;
+	const FResponse* Response = (const FResponse*)this;
 	return Response->GetUint64Checked("timestamp", 0);
 }
 
 ////////////////////////////////////////////////////////////////////////////////
 FUtf8StringView FStoreClient::FTraceInfo::GetUri() const
 {
-	const auto* Response = (const FResponse*)this;
+	const FResponse* Response = (const FResponse*)this;
 	return Response->GetString("uri", "");
 }
 
 ////////////////////////////////////////////////////////////////////////////////
 FUtf8StringView FStoreClient::FTraceInfo::GetName() const
 {
-	const auto* Response = (const FResponse*)this;
+	const FResponse* Response = (const FResponse*)this;
 	return Response->GetString("name", "nameless");
 }
 
 
 
+////////////////////////////////////////////////////////////////////////////////
+// FStoreClient::FSessionInfo
 ////////////////////////////////////////////////////////////////////////////////
 uint32 FStoreClient::FSessionInfo::GetId() const
 {
-	const auto* Response = (const FResponse*)this;
+	const FResponse* Response = (const FResponse*)this;
 	return Response->GetUint32Checked("id", 0);
 }
 
 ////////////////////////////////////////////////////////////////////////////////
 uint32 FStoreClient::FSessionInfo::GetTraceId() const
 {
-	const auto* Response = (const FResponse*)this;
+	const FResponse* Response = (const FResponse*)this;
 	return Response->GetUint32Checked("trace_id", 0);
 }
 
 ////////////////////////////////////////////////////////////////////////////////
 uint32 FStoreClient::FSessionInfo::GetIpAddress() const
 {
-	const auto* Response = (const FResponse*)this;
+	const FResponse* Response = (const FResponse*)this;
 	return Response->GetUint32Checked("ip_address", 0);
 }
 
 ////////////////////////////////////////////////////////////////////////////////
 uint32 FStoreClient::FSessionInfo::GetControlPort() const
 {
-	const auto* Response = (const FResponse*)this;
+	const FResponse* Response = (const FResponse*)this;
 	return Response->GetUint32Checked("control_port", 0);
 }
 
 
 
+////////////////////////////////////////////////////////////////////////////////
+// FStoreClient
 ////////////////////////////////////////////////////////////////////////////////
 FStoreClient* FStoreClient::Connect(const TCHAR* Host, uint32 Port)
 {
@@ -593,42 +602,42 @@ FStoreClient* FStoreClient::Connect(const TCHAR* Host, uint32 Port)
 ////////////////////////////////////////////////////////////////////////////////
 bool FStoreClient::Reconnect(const TCHAR* Host, uint32 Port)
 {
-	auto* Self = (FStoreCborClient*)this;
+	FStoreCborClient* Self = (FStoreCborClient*)this;
 	return Self->Connect(Host, static_cast<uint16>(Port));
 }
 
 ////////////////////////////////////////////////////////////////////////////////
 void FStoreClient::operator delete (void* Addr)
 {
-	auto* Self = (FStoreCborClient*)Addr;
+	FStoreCborClient* Self = (FStoreCborClient*)Addr;
 	delete Self;
 }
 
 ////////////////////////////////////////////////////////////////////////////////
 bool FStoreClient::IsValid() const
 {
-	auto* Self = (FStoreCborClient*)this;
+	const FStoreCborClient* Self = (FStoreCborClient*)this;
 	return Self->IsOpen();
 }
 
 ////////////////////////////////////////////////////////////////////////////////
 uint32 FStoreClient::GetStoreAddress() const
 {
-	auto* Self = (FStoreCborClient*)this;
+	const FStoreCborClient* Self = (FStoreCborClient*)this;
 	return Self->GetStoreAddress();
 }
 
 ////////////////////////////////////////////////////////////////////////////////
 uint32 FStoreClient::GetStorePort() const
 {
-	auto* Self = (FStoreCborClient*)this;
+	const FStoreCborClient* Self = (FStoreCborClient*)this;
 	return Self->GetStorePort();
 }
 
 ////////////////////////////////////////////////////////////////////////////////
-const FStoreClient::FStatus* FStoreClient::GetStatus()
+const FStoreClient::FStatus* FStoreClient::GetStatus() const
 {
-	auto* Self = (FStoreCborClient*)this;
+	const FStoreCborClient* Self = (FStoreCborClient*)this;
 	if (!Self->GetStatus())
 	{
 		return nullptr;
@@ -639,10 +648,9 @@ const FStoreClient::FStatus* FStoreClient::GetStatus()
 }
 
 ////////////////////////////////////////////////////////////////////////////////
-const FStoreClient::FVersion* FStoreClient::GetVersion()
+const FStoreClient::FVersion* FStoreClient::GetVersion() const
 {
-	
-	auto* Self = (FStoreCborClient*)this;
+	const FStoreCborClient* Self = (FStoreCborClient*)this;
 	if (!Self->GetVersion())
 	{
 		return nullptr;
@@ -653,9 +661,9 @@ const FStoreClient::FVersion* FStoreClient::GetVersion()
 }
 
 ////////////////////////////////////////////////////////////////////////////////
-uint32 FStoreClient::GetTraceCount()
+uint32 FStoreClient::GetTraceCount() const
 {
-	auto* Self = (FStoreCborClient*)this;
+	const FStoreCborClient* Self = (FStoreCborClient*)this;
 	if (!Self->GetTraceCount())
 	{
 		return 0;
@@ -665,9 +673,9 @@ uint32 FStoreClient::GetTraceCount()
 }
 
 ////////////////////////////////////////////////////////////////////////////////
-const FStoreClient::FTraceInfo* FStoreClient::GetTraceInfo(uint32 Index)
+const FStoreClient::FTraceInfo* FStoreClient::GetTraceInfo(uint32 Index) const
 {
-	auto* Self = (FStoreCborClient*)this;
+	const FStoreCborClient* Self = (FStoreCborClient*)this;
 	if (!Self->GetTraceInfo(Index))
 	{
 		return nullptr;
@@ -678,9 +686,9 @@ const FStoreClient::FTraceInfo* FStoreClient::GetTraceInfo(uint32 Index)
 }
 
 ////////////////////////////////////////////////////////////////////////////////
-const FStoreClient::FTraceInfo* FStoreClient::GetTraceInfoById(uint32 Id)
+const FStoreClient::FTraceInfo* FStoreClient::GetTraceInfoById(uint32 Id) const
 {
-	auto* Self = (FStoreCborClient*)this;
+	const FStoreCborClient* Self = (FStoreCborClient*)this;
 	if (!Self->GetTraceInfoById(Id))
 	{
 		return nullptr;
@@ -691,9 +699,9 @@ const FStoreClient::FTraceInfo* FStoreClient::GetTraceInfoById(uint32 Id)
 }
 
 ////////////////////////////////////////////////////////////////////////////////
-FStoreClient::FTraceData FStoreClient::ReadTrace(uint32 Id)
+FStoreClient::FTraceData FStoreClient::ReadTrace(uint32 Id) const
 {
-	auto* Self = (FStoreCborClient*)this;
+	const FStoreCborClient* Self = (FStoreCborClient*)this;
 	return FTraceData(Self->ReadTrace(Id));
 }
 
@@ -702,21 +710,21 @@ bool FStoreClient::SetStoreDirectories(const TCHAR* StoreDir,
 										const TArray<FString>& AddWatchDirs,
 										const TArray<FString>& RemoveWatchDirs)
 {
-	auto* Self = (FStoreCborClient*) this;
+	FStoreCborClient* Self = (FStoreCborClient*)this;
 	return Self->SetStoreDirectories(StoreDir, AddWatchDirs, RemoveWatchDirs);
 }
 
 ////////////////////////////////////////////////////////////////////////////////
 bool FStoreClient::SetSponsored(bool bSponsored)
 {
-	auto* Self = (FStoreCborClient*) this;
+	FStoreCborClient* Self = (FStoreCborClient*)this;
 	return Self->SetSponsored(bSponsored);
 }
-	
+
 ////////////////////////////////////////////////////////////////////////////////
 uint32 FStoreClient::GetSessionCount() const
 {
-	auto* Self = (FStoreCborClient*)this;
+	const FStoreCborClient* Self = (FStoreCborClient*)this;
 	if (!Self->GetSessionCount())
 	{
 		return 0;
@@ -728,7 +736,7 @@ uint32 FStoreClient::GetSessionCount() const
 ////////////////////////////////////////////////////////////////////////////////
 const FStoreClient::FSessionInfo* FStoreClient::GetSessionInfo(uint32 Index) const
 {
-	auto* Self = (FStoreCborClient*)this;
+	const FStoreCborClient* Self = (FStoreCborClient*)this;
 	if (!Self->GetSessionInfo(Index))
 	{
 		return nullptr;
@@ -741,7 +749,7 @@ const FStoreClient::FSessionInfo* FStoreClient::GetSessionInfo(uint32 Index) con
 ////////////////////////////////////////////////////////////////////////////////
 const FStoreClient::FSessionInfo* FStoreClient::GetSessionInfoById(uint32 Id) const
 {
-	auto* Self = (FStoreCborClient*)this;
+	const FStoreCborClient* Self = (FStoreCborClient*)this;
 	if (!Self->GetSessionInfoById(Id))
 	{
 		return nullptr;
@@ -754,7 +762,7 @@ const FStoreClient::FSessionInfo* FStoreClient::GetSessionInfoById(uint32 Id) co
 ////////////////////////////////////////////////////////////////////////////////
 const FStoreClient::FSessionInfo* FStoreClient::GetSessionInfoByTraceId(uint32 TraceId) const
 {
-	auto* Self = (FStoreCborClient*)this;
+	const FStoreCborClient* Self = (FStoreCborClient*)this;
 	if (!Self->GetSessionInfoByTraceId(TraceId))
 	{
 		return nullptr;
@@ -765,9 +773,9 @@ const FStoreClient::FSessionInfo* FStoreClient::GetSessionInfoByTraceId(uint32 T
 }
 
 ////////////////////////////////////////////////////////////////////////////////
-const FStoreClient::FSessionInfo* FStoreClient::GetSessionInfoByGuid(const FGuid& TraceGuid)
+const FStoreClient::FSessionInfo* FStoreClient::GetSessionInfoByGuid(const FGuid& TraceGuid) const
 {
-	auto* Self = (FStoreCborClient*)this;
+	const FStoreCborClient* Self = (FStoreCborClient*)this;
 	if (!Self->GetSessionInfoByTraceGuid(TraceGuid))
 	{
 		return nullptr;
@@ -776,6 +784,6 @@ const FStoreClient::FSessionInfo* FStoreClient::GetSessionInfoByGuid(const FGuid
 	const FResponse& Response = Self->GetResponse();
 	return (FSessionInfo*)(&Response);
 }
-	
+
 } // namespace Trace
 } // namespace UE
diff --git a/Engine/Source/Developer/TraceAnalysis/Public/Trace/StoreClient.h b/Engine/Source/Developer/TraceAnalysis/Public/Trace/StoreClient.h
--- a/Engine/Source/Developer/TraceAnalysis/Public/Trace/StoreClient.h
+++ b/Engine/Source/Developer/TraceAnalysis/Public/Trace/StoreClient.h
@@ -42,8 +42,6 @@ public:
 		uint64			GetSize() const;
 		uint64			GetTimestamp() const;
 		FUtf8StringView GetUri() const;
-		//const TCHAR*	GetMetadata(const TCHAR* Key) const;
-		//template <typename Lambda> uint32 ReadMetadata(Lambda&& Callback) const;
 	};
 
 	struct FTraceData
@@ -51,26 +49,22 @@ public:
 	{
 		using TUniquePtr<IInDataStream>::TUniquePtr;
 	};
-	
+
 						~FStoreClient() = default;
-	static FStoreClient*Connect(const TCHAR* Host, uint32 Port=0);
+	static FStoreClient* Connect(const TCHAR* Host, uint32 Port = 0);
 	bool				Reconnect(const TCHAR* Host, uint32 Port);
 	void				operator delete (void* Addr);
 	bool				IsValid() const;
 	uint32				GetStoreAddress() const;
 	uint32				GetStorePort() const;
-	const FStatus*		GetStatus();
-	const FVersion*		GetVersion();
-	uint32				GetTraceCount();
-	const FTraceInfo*	GetTraceInfo(uint32 Index);
-	const FTraceInfo*	GetTraceInfoById(uint32 Id);
-	FTraceData			ReadTrace(uint32 Id);
-	bool SetStoreDirectories(const TCHAR* StoreDir, const TArray<FString>& AddWatchDirs,
-							const TArray<FString>& RemoveWatchDirs);
+	const FStatus*		GetStatus() const;
+	const FVersion*		GetVersion() const;
+	uint32				GetTraceCount() const;
+	const FTraceInfo*	GetTraceInfo(uint32 Index) const;
+	const FTraceInfo*	GetTraceInfoById(uint32 Id) const;
+	FTraceData			ReadTrace(uint32 Id) const;
+	bool				SetStoreDirectories(const TCHAR* StoreDir, const TArray<FString>& AddWatchDirs, const TArray<FString>& RemoveWatchDirs);
 	bool				SetSponsored(bool bSponsored);
-#if 0
-	template <typename Lambda> uint32 GetTraceInfos(uint32 StartIndex, uint32 Count, Lambda&& Callback) const;
-#endif // 0
 
 	struct TRACEANALYSIS_API FSessionInfo
 	{
@@ -83,20 +77,7 @@ public:
 	const FSessionInfo* GetSessionInfo(uint32 Index) const;
 	const FSessionInfo* GetSessionInfoById(uint32 Id) const;
 	const FSessionInfo* GetSessionInfoByTraceId(uint32 TraceId) const;
-	const FSessionInfo* GetSessionInfoByGuid(const FGuid& TraceGuid);
-#if 0
-	template <typename Lambda> uint32	GetSessionInfos(uint32 StartIndex, uint32 Count, Lambda&& Callback) const;
-#endif // 0
-
-#if 0
-	// -------
-	class IStoreSubscriber
-	{
-		virtual void	OnStoreEvent() = 0;
-	};
-	bool				Subscribe(IStoreSubscriber* Subscriber);
-	bool				Unsubscribe(IStoreSubscriber* Subscriber);
-#endif // 0
+	const FSessionInfo* GetSessionInfoByGuid(const FGuid& TraceGuid) const;
 
 private:
 						FStoreClient() = default;
