Fix narrowing of store directory string

At one point it was assumed that store directory would only ever be limited to ANSI. The selected directory string was narrowed before being sent to UTS.

diff --git a/Engine/Source/Developer/TraceAnalysis/Private/Store/CborPayload.h b/Engine/Source/Developer/TraceAnalysis/Private/Store/CborPayload.h
--- a/Engine/Source/Developer/TraceAnalysis/Private/Store/CborPayload.h
+++ b/Engine/Source/Developer/TraceAnalysis/Private/Store/CborPayload.h
@@ -73,6 +73,7 @@ public:
 	template <int N>			TPayloadBuilder(const char (&Path)[N]);
 	template <int N> void		AddInteger(const char (&Name)[N], int64 Value);
 	template <int N> void		AddString(const char (&Name)[N], const char* Value, int32 Length=-1);
+	template <int N> void		AddString(const char (&Name)[N], const FStringView& String);
 	template <int N> void		AddStringArray(const char(&Name)[N], const TArray<FString>& Values);
 	template <int N, typename T>
 					void		AddArray(const char(&Name)[N], const TArrayView<T>& Values);
@@ -122,7 +123,17 @@ inline void TPayloadBuilder<Size>::AddString(
 	CborWriter.WriteValue(Name, N - 1);
 	CborWriter.WriteValue(Value, Length);
 }
+
+////////////////////////////////////////////////////////////////////////////////
+template <int Size>
+template <int N>
+void TPayloadBuilder<Size>::AddString(const char(& Name)[N], const FStringView& String)
+{
 	
+	CborWriter.WriteValue(Name, N - 1);
+	CborWriter.WriteValue(String);
+}
+
 ////////////////////////////////////////////////////////////////////////////////
 template<int Size>
 template<int N>
diff --git a/Engine/Source/Developer/TraceAnalysis/Private/Store/StoreClient.cpp b/Engine/Source/Developer/TraceAnalysis/Private/Store/StoreClient.cpp
--- a/Engine/Source/Developer/TraceAnalysis/Private/Store/StoreClient.cpp
+++ b/Engine/Source/Developer/TraceAnalysis/Private/Store/StoreClient.cpp
@@ -405,7 +405,7 @@ bool FStoreCborClient::SetStoreDirectories(const TCHAR* StoreDir, const TArray<F
 	TPayloadBuilder<> Builder("v1/settings/write");
 	if (StoreDir)
 	{
-		Builder.AddString("StoreDir", TCHAR_TO_ANSI(StoreDir));
+		Builder.AddString("StoreDir", FStringView(StoreDir));
 	}
 	TArray<FString> WatchDirs;
 	if (!RemoveWatchDir.IsEmpty())
diff --git a/Engine/Source/Runtime/Cbor/Private/CborWriter.cpp b/Engine/Source/Runtime/Cbor/Private/CborWriter.cpp
--- a/Engine/Source/Runtime/Cbor/Private/CborWriter.cpp
+++ b/Engine/Source/Runtime/Cbor/Private/CborWriter.cpp
@@ -121,11 +121,16 @@ void FCborWriter::WriteValue(double Value)
 }
 
 void FCborWriter::WriteValue(const FString& Value)
+{
+	WriteValue(FStringView(Value));
+}
+
+void FCborWriter::WriteValue(const FStringView& Value)
 {
 	ScopedCborArchiveEndianness ScopedArchiveEndianness(*Stream, Endianness);
 
 	CheckContext(ECborCode::TextString);
-	FTCHARToUTF8 UTF8String(*Value);
+	const FTCHARToUTF8 UTF8String(Value);
 	// Write string header
 	WriteUIntValue(ECborCode::TextString, *Stream, (uint64)UTF8String.Length());
 	// Write string
diff --git a/Engine/Source/Runtime/Cbor/Public/CborWriter.h b/Engine/Source/Runtime/Cbor/Public/CborWriter.h
--- a/Engine/Source/Runtime/Cbor/Public/CborWriter.h
+++ b/Engine/Source/Runtime/Cbor/Public/CborWriter.h
@@ -49,6 +49,7 @@ public:
 	CBOR_API void WriteValue(float Value);
 	CBOR_API void WriteValue(double Value);
 	CBOR_API void WriteValue(const FString& Value);
+	CBOR_API void WriteValue(const FStringView& Value);
 	CBOR_API void WriteValue(const char* CString, uint64 Length);
 	CBOR_API void WriteValue(const uint8* Bytes, uint64 Length);
 
