[Insights]
  - TraceAnalysis: Added debug functionality to allow recover data from traces with missing packages / missing serial sync events (code disabled by default).
  - TraceAnalysis: Added more detailed debug info for per thread buffers' data size (parsed data size + remaining data size).
  - TraceAnalysis: Added value of uncompressed data size not processed at end of analysis in the warning message.
  - TraceAnalysis: Fixed debug values re number of "parsed events" to include the Terminal events (to match with total number of "parsed events").
  - TraceAnalysis: Fixed few comments in analysis engine code.

diff --git a/Engine/Source/Developer/TraceAnalysis/Private/Analysis/Engine.cpp b/Engine/Source/Developer/TraceAnalysis/Private/Analysis/Engine.cpp
--- a/Engine/Source/Developer/TraceAnalysis/Private/Analysis/Engine.cpp
+++ b/Engine/Source/Developer/TraceAnalysis/Private/Analysis/Engine.cpp
@@ -1811,6 +1811,7 @@ struct FAnalysisState
 	int32 MaxEventDescs = 0; // for protocol 5
 	int32 SerialWrappedCount = 0; // for protocol 5
 	int32 NumSkippedSerialGaps = 0; // for protocol 5
+	uint64 NumSkippedSerials = 0; // for protocol 5
 
 	uint64 TotalEventCount = 0;
 	uint64 NewEventCount = 0;
@@ -1924,7 +1925,7 @@ void FTraceAnalyzer::OnNewTrace(const FOnEventContext& Context)
 	const FEventData& EventData = Context.EventData;
 
 	// "Serial" will tell us approximately where we've started in the log serial
-	// range. We'll bias it by half so we won't accept any serialised events and
+	// range. We'll bias it by half so we won't accept any serialized events and
 	// mark the MSB to indicate that the current serial should be corrected.
 	auto& Serial = State.Serial;
 	uint32 Hint = EventData.GetValue<uint32>("Serial");
@@ -2266,7 +2267,7 @@ void FAnalysisBridge::DispatchLeaveScope()
 	int64 ScopeValue = int64(ThreadInfo->ScopeRoutes.Pop(EAllowShrinking::No));
 	if (ScopeValue < 0)
 	{
-		// enter/leave pair without an event inbetween.
+		// enter/leave pair without an event in-between.
 		return;
 	}
 
@@ -2948,7 +2949,7 @@ FProtocol2Stage::EStatus FProtocol2Stage::OnData(
 			continue;
 		}
 
-		// If we didn't stumble across the next serialised event we have done all
+		// If we didn't stumble across the next serialized event we have done all
 		// we can for now.
 		if ((uint32(MinLogSerial - Serial.Value) & Serial.Mask) != 0)
 		{
@@ -3142,7 +3143,7 @@ int32 FProtocol2Stage::OnDataAux(FStreamReader& Reader, FAuxDataCollector& Colle
 			return 0;
 		}
 
-		// Is the following sequence a blob of auxilary data or the null
+		// Is the following sequence a blob of auxiliary data or the null
 		// terminator byte?
 		if (NextByte[0] == 0)
 		{
@@ -3150,7 +3151,7 @@ int32 FProtocol2Stage::OnDataAux(FStreamReader& Reader, FAuxDataCollector& Colle
 			return 1;
 		}
 
-		// Get header and the auxilary blob's size
+		// Get header and the auxiliary blob's size
 		const auto* Header = Reader.GetPointer<Protocol1::FAuxHeader>();
 		if (Header == nullptr)
 		{
@@ -3980,6 +3981,46 @@ FProtocol5Stage::EStatus FProtocol5Stage::OnDataNormal(const FMachineContext& Co
 		while (NumAvailableEvents > MaxAvailableEventsLowLimit);
 	}
 
+#if 0
+	// If the reader is empty (no more incoming trace data) and we still have parsed events not dispatched,
+	// then we'll try to skip the missing serial sync events.
+	if (Transport.GetReader()->IsEmpty() && NumAvailableEvents > 0)
+	{
+		UE_TRACE_ANALYSIS_DEBUG_LOG("Error: Trace reader is empty, but there are still %d parsed events not dispatched. Starting to skip the missing serial sync events!", NumAvailableEvents);
+		Context.EmitMessagef(
+			EAnalysisMessageSeverity::Error,
+			TEXT("Trace reader is empty, but there are still %d parsed events not dispatched. Starting to skip the missing serial sync events!"), NumAvailableEvents);
+		while (true)
+		{
+			// Skip serials and continue to dispatch parsed events.
+			bSkipSerial = true;
+			NextSerialWaitCount = 0;
+			int32 NumDispatchedEvents = DispatchNormalEvents(Context, EventDescHeap);
+			if (NumDispatchedEvents < 0)
+			{
+				return EStatus::Error;
+			}
+#if UE_TRACE_ANALYSIS_DEBUG && UE_TRACE_ANALYSIS_DEBUG_LEVEL >= 2
+			UE_TRACE_ANALYSIS_DEBUG_LOG("Skipped serials and dispatched %d normal events (%d --> %d)", NumDispatchedEvents, NumAvailableEvents, NumAvailableEvents - NumDispatchedEvents);
+#endif
+			NumAvailableEvents -= NumDispatchedEvents;
+			check(NumAvailableEvents >= 0);
+
+			if (NumAvailableEvents == 0)
+			{
+				break;
+			}
+			if (NumDispatchedEvents == 0)
+			{
+				UE_TRACE_ANALYSIS_DEBUG_LOG("Cannot dispatch the remaining %d (incomplete) parsed events!", NumAvailableEvents);
+				break;
+			}
+		}
+		bNotEnoughData = false;
+		bSkipSerial = false;
+	}
+#endif
+
 	// If there are any streams left in the heap then we are unable to proceed
 	// until more data is received. We'll rewind the streams until more data is
 	// available. It is not an efficient way to do things, but it is simple way.
@@ -4010,9 +4051,8 @@ int32 FProtocol5Stage::DispatchNormalEvents(const FMachineContext& Context, TArr
 	UE_TRACE_ANALYSIS_DEBUG_LOG("Queued event descs: %d", EventDescs.Num());
 #endif
 
-	int32 NumDispatchedUnsyncEvents = 0;
-
 	// Process leading unsynchronized events so that each stream starts with a synchronized event.
+	int32 NumDispatchedUnsyncEvents = 0;
 	for (FEventDescStream& Stream : EventDescHeap)
 	{
 		// Extract a run of consecutive unsynchronized events
@@ -4039,15 +4079,17 @@ int32 FProtocol5Stage::DispatchNormalEvents(const FMachineContext& Context, TArr
 	}
 
 	// Trim off empty streams
+	int32 NumTerminators = EventDescHeap.Num();
 	EventDescHeap.RemoveAllSwap([] (const FEventDescStream& Stream)
 	{
 		return (Stream.EventDescs->Serial == ESerial::Terminal);
 	});
+	NumTerminators -= EventDescHeap.Num();
 
 	// Early out if there isn't any events available.
 	if (UNLIKELY(EventDescHeap.IsEmpty()))
 	{
-		return NumDispatchedUnsyncEvents;
+		return NumDispatchedUnsyncEvents + NumTerminators;
 	}
 
 	// A min-heap is used to peel off groups of events by lowest serial
@@ -4066,7 +4108,11 @@ int32 FProtocol5Stage::DispatchNormalEvents(const FMachineContext& Context, TArr
 		uint32 LowestSerial = EventDescHeap.HeapTop().EventDescs[0].Serial;
 		if (LowestSerial != NextSerial)
 		{
-			UE_TRACE_ANALYSIS_DEBUG_LOG("Warning: NextSerial skips %lld events (from %u to %u)", (int64)LowestSerial - (int64)NextSerial, NextSerial, LowestSerial);
+#if UE_TRACE_ANALYSIS_DEBUG
+			uint32 SerialsToSkip = ((LowestSerial & ESerial::Mask) + ESerial::Range - (NextSerial & ESerial::Mask)) & ESerial::Mask;
+			Context.Bridge.GetState().NumSkippedSerials += SerialsToSkip;
+			UE_TRACE_ANALYSIS_DEBUG_LOG("Warning: NextSerial skips %u events (from %u to %u)", SerialsToSkip, NextSerial, LowestSerial);
+#endif
 			NextSerial = LowestSerial;
 		}
 	}
@@ -4090,7 +4136,7 @@ int32 FProtocol5Stage::DispatchNormalEvents(const FMachineContext& Context, TArr
 		return -1;
 	}
 
-	return NumDispatchedUnsyncEvents + NumDispatchedEvents;
+	return NumDispatchedUnsyncEvents + NumTerminators + NumDispatchedEvents;
 }
 
 ////////////////////////////////////////////////////////////////////////////////
@@ -4138,6 +4184,7 @@ int32 FProtocol5Stage::DispatchEvents(
 			UE_TRACE_ANALYSIS_DEBUG_LOG("Skip serial gap (%u +%u) --> NextSerial=%u", EndDesc->Serial, EndDesc->GapLength, NextSerial);
 #endif
 			UpdateHeap(Stream, EndDesc + 1);
+			NumDispatchedEvents += ((EndDesc + 1)->Serial == ESerial::Terminal) ? 2 : 1;// GapThreadId event + Terminal event
 			continue;
 		}
 
@@ -4187,7 +4234,8 @@ int32 FProtocol5Stage::DispatchEvents(
 			for (const FEventDescStream& EventDescStream : EventDescHeap)
 			{
 				uint32 BufferSize = 0;
-				uint32 DataSize = 0;
+				uint32 DataSize = 0; // parsed data size + remaining data size
+				uint32 RemainingSize = 0;
 				FTidPacketTransport* InnerTransport = (FTidPacketTransport*)(&Transport);
 				for (uint32 i = 0, n = InnerTransport->GetThreadCount(); i < n; ++i)
 				{
@@ -4196,24 +4244,31 @@ int32 FProtocol5Stage::DispatchEvents(
 					if (ThreadId == EventDescStream.ThreadId)
 					{
 						BufferSize = ThreadReader->GetBufferSize();
-						DataSize = ThreadReader->GetRemaining();
+						const FEventDesc& EventDesc = EventDescStream.EventDescs[0];
+						uint32 HeaderSize = 1 + EventDesc.bTwoByteUid + (ESerial::Bits / 8);
+						DataSize = ThreadReader->GetBacktrackSize(EventDesc.Data - HeaderSize);
+						RemainingSize = ThreadReader->GetRemaining();
+						break;
 					}
 				}
-				if (EventDescStream.EventDescs->Serial == NextSerial)
+				UE_TRACE_ANALYSIS_DEBUG_BeginStringBuilder();
+				UE_TRACE_ANALYSIS_DEBUG_Appendf("  Tid=%u : Serial=%u BufferSize=%u DataSize=%u", EventDescStream.ThreadId, EventDescStream.EventDescs->Serial, BufferSize, DataSize);
+				if (RemainingSize != 0)
 				{
-					UE_TRACE_ANALYSIS_DEBUG_LOG("  Tid=%u : Serial=%u BufferSize=%u DataSize=%u (next)", EventDescStream.ThreadId, EventDescStream.EventDescs->Serial, BufferSize, DataSize);
+					UE_TRACE_ANALYSIS_DEBUG_Appendf(" (%u + %u)", DataSize - RemainingSize, RemainingSize);
 				}
-				else
+				if (EventDescStream.EventDescs->Serial == NextSerial)
 				{
-					UE_TRACE_ANALYSIS_DEBUG_LOG("  Tid=%u : Serial=%u BufferSize=%u DataSize=%u", EventDescStream.ThreadId, EventDescStream.EventDescs->Serial, BufferSize, DataSize);
+					UE_TRACE_ANALYSIS_DEBUG_Append(" (next)");
 				}
+				UE_TRACE_ANALYSIS_DEBUG_EndStringBuilder();
 			}
 #endif // UE_TRACE_ANALYSIS_DEBUG_LEVEL
 #endif // UE_TRACE_ANALYSIS_DEBUG
 
 #if 0
 			int32 MinSerial = EndDesc->Serial;
-			EventDescHeap.Heapify();
+			EventDescHeap.Heapify(FSerialDistancePredicate{NextSerial});
 			const FEventDescStream& MinStream = EventDescHeap.HeapTop();
 			const FEventDesc* MinDesc = MinStream.EventDescs;
 			if (MinDesc->Serial < MinSerial)
@@ -4239,6 +4294,7 @@ int32 FProtocol5Stage::DispatchEvents(
 		}
 
 		UpdateHeap(Stream, EndDesc);
+		NumDispatchedEvents += (EndDesc->Serial == ESerial::Terminal) ? 1 : 0; // Terminal event
 	}
 	while (!EventDescHeap.IsEmpty());
 
@@ -4311,7 +4367,7 @@ int32 FProtocol5Stage::ParseEvents(FStreamReader& Reader, EventDescArray& OutEve
 			if (Ok == 0)
 			{
 #if UE_TRACE_ANALYSIS_DEBUG && UE_TRACE_ANALYSIS_DEBUG_LEVEL >= 2
-				UE_TRACE_ANALYSIS_DEBUG_LOG("Warning: Incomplete aux stack! Rewind %d parsed events.", OutEventDescs.Num() - RewindDescsNum);
+				UE_TRACE_ANALYSIS_DEBUG_LOG("Warning: Incomplete aux stack (Uid=%u)! Rewind %d parsed events.", EventDesc.Uid, OutEventDescs.Num() - RewindDescsNum);
 #endif
 				OutEventDescs.SetNum(RewindDescsNum);
 				Reader.RestoreMark(RewindMark);
@@ -4594,7 +4650,7 @@ void FProtocol5Stage::DetectSerialGaps(TArray<FEventDescStream>& EventDescHeap)
 	// the trace tail to make space for new trace events, and 2) when Trace's
 	// worker thread ticks, samples all the trace buffers and sends their data.
 	// In late-connect scenarios these gaps need to be skipped over in order to
-	// successfully reserialize events in the data stream. To further complicate
+	// successfully re-serialize events in the data stream. To further complicate
 	// matters, most of the gaps from (2) will get filled by the following update,
 	// leading to initial false positive gaps. By embedding sync points in the
 	// stream we can reliably differentiate genuine gaps from temporary ones.
@@ -4923,7 +4979,14 @@ void FProtocol7Stage::ExitStage(const FMachineContext& Context)
 	// Ensure the transport does not have pending buffers (i.e. event data not yet processed).
 	if (!Transport.IsEmpty())
 	{
-		Context.EmitMessage(EAnalysisMessageSeverity::Warning, TEXT("Transport buffers are not empty at end of analysis (protocol 7)!"));
+		uint64 TotalRemainingDataSize = 0;
+		for (int32 ThreadIndex = 0, ThreadCount = Transport.GetThreadCount(); ThreadIndex < ThreadCount; ++ThreadIndex)
+		{
+			TotalRemainingDataSize += Transport.GetThreadStream(ThreadIndex)->GetRemaining();
+		}
+		Context.EmitMessagef(
+			EAnalysisMessageSeverity::Warning,
+			TEXT("Transport buffers are not empty at end of analysis (%llu bytes; protocol 7)!"), TotalRemainingDataSize);
 	}
 
 #if UE_TRACE_ANALYSIS_DEBUG
@@ -5295,6 +5358,7 @@ void FAnalysisEngine::FImpl::End()
 	UE_TRACE_ANALYSIS_DEBUG_LOG("Serial.Value: %u (0x%X)", State.Serial.Value, State.Serial.Value);
 	UE_TRACE_ANALYSIS_DEBUG_LOG("MaxEventDescs: %d", State.MaxEventDescs);
 	UE_TRACE_ANALYSIS_DEBUG_LOG("SkippedSerialGaps: %d", State.NumSkippedSerialGaps);
+	UE_TRACE_ANALYSIS_DEBUG_LOG("SkippedSerials: %llu", State.NumSkippedSerials);
 
 	uint32 Digits = 0;
 	uint64 Value = State.TotalEventSize;
diff --git a/Engine/Source/Developer/TraceAnalysis/Private/Analysis/StreamReader.cpp b/Engine/Source/Developer/TraceAnalysis/Private/Analysis/StreamReader.cpp
--- a/Engine/Source/Developer/TraceAnalysis/Private/Analysis/StreamReader.cpp
+++ b/Engine/Source/Developer/TraceAnalysis/Private/Analysis/StreamReader.cpp
@@ -7,7 +7,7 @@ namespace UE {
 namespace Trace {
 
 ////////////////////////////////////////////////////////////////////////////////
-const uint8* FStreamReader::GetPointer(uint32 Size)
+const uint8* FStreamReader::GetPointer(uint32 Size) const
 {
 	if (Cursor + Size > End)
 	{
diff --git a/Engine/Source/Developer/TraceAnalysis/Private/Analysis/StreamReader.h b/Engine/Source/Developer/TraceAnalysis/Private/Analysis/StreamReader.h
--- a/Engine/Source/Developer/TraceAnalysis/Private/Analysis/StreamReader.h
+++ b/Engine/Source/Developer/TraceAnalysis/Private/Analysis/StreamReader.h
@@ -14,45 +14,51 @@ class FStreamReader
 {
 public:
 	template <typename Type>
-	Type const*					GetPointer();
+	Type const*					GetPointer() const;
 	template <typename Type>
-	Type const*					GetPointerUnchecked();
-	const uint8*				GetPointer(uint32 Size);
-	const uint8*				GetPointerUnchecked();
+	Type const*					GetPointerUnchecked() const;
+	const uint8*				GetPointer(uint32 Size) const;
+	const uint8*				GetPointerUnchecked() const;
 	void						Advance(uint32 Size);
 	bool						IsEmpty() const;
-	uint32						GetRemaining() const;
 	bool						CanMeetDemand() const;
+	uint32						GetRemaining() const;
+	uint32						GetBacktrackSize(const uint8* To) const;
 	bool						Backtrack(const uint8* To);
 	struct FMark*				SaveMark() const;
 	void						RestoreMark(struct FMark* Mark);
 
 protected:
 	uint8*						Buffer = nullptr;
-	uint32						DemandHint = 0;
+	mutable uint32				DemandHint = 0;
 	uint32						Cursor = 0;
 	uint32						End = 0;
 };
 
 ////////////////////////////////////////////////////////////////////////////////
 template <typename Type>
-Type const* FStreamReader::GetPointer()
+Type const* FStreamReader::GetPointer() const
 {
 	return (Type const*)GetPointer(sizeof(Type));
 }
 
 ////////////////////////////////////////////////////////////////////////////////
 template <typename Type>
-Type const* FStreamReader::GetPointerUnchecked()
+Type const* FStreamReader::GetPointerUnchecked() const
 {
 	return (Type const*)GetPointerUnchecked();
 }
 
 ////////////////////////////////////////////////////////////////////////////////
-inline const uint8* FStreamReader::GetPointerUnchecked()
+inline const uint8* FStreamReader::GetPointerUnchecked() const
 {
 	return Buffer + Cursor;
 }
+////////////////////////////////////////////////////////////////////////////////
+inline uint32 FStreamReader::GetBacktrackSize(const uint8* To) const
+{
+	return End - uint32(UPTRINT(To - Buffer));
+}
 
 
 
diff --git a/Engine/Source/Developer/TraceAnalysis/Private/Analysis/Transport/Transport.h b/Engine/Source/Developer/TraceAnalysis/Private/Analysis/Transport/Transport.h
--- a/Engine/Source/Developer/TraceAnalysis/Private/Analysis/Transport/Transport.h
+++ b/Engine/Source/Developer/TraceAnalysis/Private/Analysis/Transport/Transport.h
@@ -13,6 +13,7 @@ class FTransport
 public:
 	virtual					~FTransport() {}
 	void					SetReader(FStreamReader& InReader);
+	FStreamReader*			GetReader() const { return Reader; }
 	template <typename RetType>
 	RetType const*			GetPointer();
 	template <typename RetType>
