[Insights] TraceAnalysis: Fixed assert for late connect sessions (additional events for serial gaps not being accounted when computing total available parsed events).

diff --git a/Engine/Source/Developer/TraceAnalysis/Private/Analysis/Engine.cpp b/Engine/Source/Developer/TraceAnalysis/Private/Analysis/Engine.cpp
--- a/Engine/Source/Developer/TraceAnalysis/Private/Analysis/Engine.cpp
+++ b/Engine/Source/Developer/TraceAnalysis/Private/Analysis/Engine.cpp
@@ -3919,6 +3919,7 @@ FProtocol5Stage::EStatus FProtocol5Stage::OnDataNormal(const FMachineContext& Co
 		{
 			return EStatus::Error;
 		}
+		NumAvailableEvents += SerialGaps.Num(); // serial gaps are detected during DispatchNormalEvents
 #if UE_TRACE_ANALYSIS_DEBUG && UE_TRACE_ANALYSIS_DEBUG_LEVEL >= 2
 		UE_TRACE_ANALYSIS_DEBUG_LOG("Dispatched %d normal events (%d --> %d)", NumDispatchedEvents, NumAvailableEvents, NumAvailableEvents - NumDispatchedEvents);
 #endif
