[Insights] Added validation and warning log message for trace anaysis decoding a packet that is empty (ex: for cases when runtime code fails to encode a large packet). Trace analysis is allowed to continue when encounters an empty trace packet.

diff --git a/Engine/Source/Developer/TraceAnalysis/Private/Analysis/Transport/TidPacketTransport.cpp b/Engine/Source/Developer/TraceAnalysis/Private/Analysis/Transport/TidPacketTransport.cpp
--- a/Engine/Source/Developer/TraceAnalysis/Private/Analysis/Transport/TidPacketTransport.cpp
+++ b/Engine/Source/Developer/TraceAnalysis/Private/Analysis/Transport/TidPacketTransport.cpp
@@ -101,7 +101,7 @@ FTidPacketTransport::EReadPacketResult FTidPacketTransport::ReadPacket()
 			DataSize,
 			DecodedSize,
 			((double)DataSize * 100.0) / (double)DecodedSize - 100.0,
-			DataSize > DecodedSize ? " !!!" : "");
+			((DataSize == 0) || (DataSize > DecodedSize)) ? " !!!" : "");
 #endif // UE_TRACE_ANALYSIS_DEBUG_LEVEL
 		TotalPacketHeaderSize += sizeof(FTidPacketEncoded) - sizeof(FTidPacketBase);
 		TotalDecodedSize += DecodedSize;
@@ -109,25 +109,36 @@ FTidPacketTransport::EReadPacketResult FTidPacketTransport::ReadPacket()
 		TotalDataSizePerThread += DecodedSize;
 #endif // UE_TRACE_ANALYSIS_DEBUG
 
-		uint8* Dest = Thread->Buffer.Append(DecodedSize);
 		DataSize -= sizeof(DecodedSize);
-		int32 ResultSize = UE::Trace::Private::Decode(Packet->Data, DataSize, Dest, DecodedSize);
-		if (int32(DecodedSize) != ResultSize)
+
+		// for debugging purposes only
+		//UE_LOG(LogCore, Log, TEXT("Decompress packet (tid=%u): %u --> %u bytes"), ThreadId, DataSize, uint32(DecodedSize));
+
+		if (DataSize != 0)
 		{
-			UE_LOG(LogCore, Error, TEXT("Unable to decompress packet, expected %d bytes but decoded %d bytes."), DecodedSize, ResultSize);
-			return EReadPacketResult::ReadError;
-		}
+			uint8* Dest = Thread->Buffer.Append(DecodedSize);
+			int32 ResultSize = UE::Trace::Private::Decode(Packet->Data, DataSize, Dest, DecodedSize);
+			if (int32(DecodedSize) != ResultSize)
+			{
+				UE_LOG(LogCore, Error, TEXT("Unable to decompress packet; tid=%u, expected %u bytes, but decoded %d bytes!"), ThreadId, uint32(DecodedSize), ResultSize);
+				return EReadPacketResult::ReadError;
+			}
 
 #if UE_TRACE_ANALYSIS_DEBUG && UE_TRACE_ANALYSIS_DEBUG_LEVEL >= 4
-		UE_TRACE_ANALYSIS_DEBUG_BeginStringBuilder();
-		for (uint32 i = 0; i < 32 && i < DecodedSize; ++i)
+			UE_TRACE_ANALYSIS_DEBUG_BeginStringBuilder();
+			for (uint32 i = 0; i < 32 && i < DecodedSize; ++i)
+			{
+				UE_TRACE_ANALYSIS_DEBUG_Appendf("%02X ", Dest[i]);
+			}
+			UE_TRACE_ANALYSIS_DEBUG_EndStringBuilder();
+#endif // UE_TRACE_ANALYSIS_DEBUG
+		}
+		else // DataSize == 0
 		{
-			UE_TRACE_ANALYSIS_DEBUG_Appendf("%02X ", Dest[i]);
+			UE_LOG(LogCore, Warning, TEXT("Unable to decompress packet; tid=%u, expected to decompress %u bytes, but the packet is empty!"), ThreadId, uint32(DecodedSize));
 		}
-		UE_TRACE_ANALYSIS_DEBUG_EndStringBuilder();
-#endif // UE_TRACE_ANALYSIS_DEBUG
 	}
-	else
+	else // not encoded
 	{
 #if UE_TRACE_ANALYSIS_DEBUG
 #if UE_TRACE_ANALYSIS_DEBUG_LEVEL >= 3
@@ -143,6 +154,9 @@ FTidPacketTransport::EReadPacketResult FTidPacketTransport::ReadPacket()
 		TotalDataSizePerThread += DataSize;
 #endif // UE_TRACE_ANALYSIS_DEBUG
 
+		// for debugging purposes only
+		//UE_LOG(LogCore, Log, TEXT("Uncompressed packet (tid=%u): %u bytes"), ThreadId, DataSize);
+
 		Thread->Buffer.Append((uint8*)(PacketBase + 1), DataSize);
 
 #if UE_TRACE_ANALYSIS_DEBUG && UE_TRACE_ANALYSIS_DEBUG_LEVEL >= 4
@@ -156,7 +170,7 @@ FTidPacketTransport::EReadPacketResult FTidPacketTransport::ReadPacket()
 	}
 
 #if UE_TRACE_PACKET_VERIFICATION
-	if(bHasPacketSerial)
+	if (bHasPacketSerial)
 	{
 		uint64 PacketSerial = *GetPointer<uint64>();
 		if ((LastPacketSerial + 1) != PacketSerial)
