Direct socket connection for analysis.

Adds a data stream type for Trace analysis which allows a runtime to connect directly to the analysis session, bypassing UnrealTraceServer. The session will analyzed directly in memory and not stored on disk. The direct socket connection is using different port than the standard networked trace connection, to avoid collisions with UTS. It is only available for local connections and in the same process.

diff --git a/Engine/Source/Developer/TraceAnalysis/Private/Analysis/Processor.cpp b/Engine/Source/Developer/TraceAnalysis/Private/Analysis/Processor.cpp
--- a/Engine/Source/Developer/TraceAnalysis/Private/Analysis/Processor.cpp
+++ b/Engine/Source/Developer/TraceAnalysis/Private/Analysis/Processor.cpp
@@ -38,6 +38,11 @@ uint32 FAnalysisProcessor::FImpl::Run()
 {
 	LLM_SCOPE_BYNAME(TEXT("TraceAnalysis"));
 
+	// Allow async streams to establish the data connection before
+	// starting the read loop.
+	while (!StopEvent->Wait(0, true) && !DataStream.WaitUntilReady())
+	{}
+
 	AnalysisEngine.Begin();
 
 	FStreamBuffer Buffer(4 << 20);
diff --git a/Engine/Source/Developer/TraceAnalysis/Private/DataStream.cpp b/Engine/Source/Developer/TraceAnalysis/Private/DataStream.cpp
--- a/Engine/Source/Developer/TraceAnalysis/Private/DataStream.cpp
+++ b/Engine/Source/Developer/TraceAnalysis/Private/DataStream.cpp
@@ -1,12 +1,22 @@
 // Copyright Epic Games, Inc. All Rights Reserved.
 
 #include "Trace/DataStream.h"
+#include "DataStreamInternal.h"
+
 #include "GenericPlatform/GenericPlatformFile.h"
+#include "HAL/Event.h"
 #include "HAL/PlatformFileManager.h"
+#include "HAL/PlatformProcess.h"
+#include "HAL/Runnable.h"
+#include "HAL/RunnableThread.h"
 
 namespace UE {
 namespace Trace {
 
+//--------------------------------------------------------------------
+// FFileDataStream
+//--------------------------------------------------------------------
+	
 FFileDataStream::FFileDataStream()
 	: Handle(nullptr)
 	, Remaining(0)
@@ -60,5 +70,239 @@ void FFileDataStream::Close()
 	Handle.Reset();
 }
 
+
+//--------------------------------------------------------------------
+// FTraceDataStream
+//--------------------------------------------------------------------
+	
+FTraceDataStream::FTraceDataStream(asio::ip::tcp::socket& InSocket)
+: Socket(MoveTemp(InSocket))
+{
+	asio::socket_base::receive_buffer_size RecvBufferSize(4 << 20);
+	Socket.set_option(RecvBufferSize);
+}
+
+FTraceDataStream::~FTraceDataStream()
+{
+	FTraceDataStream::Close();
+}
+
+bool FTraceDataStream::IsOpen() const
+{
+	return Socket.is_open();
+}
+
+void FTraceDataStream::Close()
+{
+	Socket.shutdown(asio::ip::tcp::socket::shutdown_receive);
+	Socket.close();
+}
+
+int32 FTraceDataStream::Read(void* Dest, uint32 DestSize)
+{
+	auto Handle = Socket.native_handle();
+
+	fd_set Fds;
+	FD_ZERO(&Fds);
+	FD_SET(Handle, &Fds);
+
+	timeval Timeout;
+	Timeout.tv_sec = 1;
+	Timeout.tv_usec = 0;
+
+	while (true)
+	{
+		fd_set ReadFds = Fds;
+		int Ret = select((int)Handle + 1, &ReadFds, 0, 0, &Timeout);
+
+		if (Ret < 0)
+		{
+			Close();
+			return -1;
+		}
+
+		if (Ret == 0)
+		{
+			continue;
+		}
+
+		asio::error_code ErrorCode;
+		size_t BytesRead = Socket.read_some(asio::buffer(Dest, DestSize), ErrorCode);
+		if (ErrorCode)
+		{
+			Close();
+			return -1;
+		}
+
+		return int32(BytesRead);
+	}
+}
+	
+//--------------------------------------------------------------------
+// FDirectSocketContext
+//--------------------------------------------------------------------
+
+/**
+* Small utility class to avoid exposing the asio types in the public
+* header.
+*/
+struct FDirectSocketContext
+{
+	FDirectSocketContext()
+		: Context(1)
+		, Acceptor(Context)
+	{}
+	
+	asio::io_context Context;
+	asio::ip::tcp::acceptor Acceptor;
+};
+
+//--------------------------------------------------------------------
+// FDirectSocketStream
+//--------------------------------------------------------------------
+	
+FDirectSocketStream::FDirectSocketStream()
+	: ConnectionEvent(FPlatformProcess::GetSynchEventFromPool(false))
+{
+	AsioContext = MakeUnique<FDirectSocketContext>();
+}
+
+
+FDirectSocketStream::~FDirectSocketStream()
+{
+	FDirectSocketStream::Close();
+	// Make sure structures that rely on the asio context are
+	// destructed before it is removed.
+	ListeningThread.Reset();
+	InternalStream.Reset();
+}
+	
+
+uint16 FDirectSocketStream::StartListening()
+{
+	uint16 Port = DefaultPort;
+	// Try to bind the default port. If that is busy move to next candidate
+	while (!CreateSocket(Port) && (Port - DefaultPort) < MaxPortAttempts)
+	{
+		++Port;
+	}
+	if (Port - DefaultPort >= MaxPortAttempts)
+	{
+		return 0;
+	}
+
+	// We cannot block on this thread or on the analysis thread (supports cancellation), create
+	// a short-lived thread for the blocking accept call.
+	ListeningThread = TUniquePtr<FRunnableThread>(FRunnableThread::Create(this, TEXT("FDirectSocketStreamListener")));
+	UE_LOG(LogCore, Log, TEXT("Started listening thread for direct trace connection on port %u"), Port);
+
+	return Port;
+}
+
+
+bool FDirectSocketStream::CreateSocket(uint16 Port)
+{
+	using asio::ip::tcp;
+
+	tcp::acceptor Acceptor(AsioContext->Context);
+
+#if PLATFORM_WINDOWS
+	DWORD Flags = WSA_FLAG_NO_HANDLE_INHERIT|WSA_FLAG_OVERLAPPED;
+	SOCKET Socket = WSASocketW(AF_INET, SOCK_STREAM, IPPROTO_TCP, nullptr, 0, Flags);
+	Acceptor.assign(tcp::v4(), Socket);
+#else
+	Acceptor.open(tcp::v4());
+#endif
+
+	// Reuse address to avoid locking up port
+	asio::socket_base::reuse_address reuse_option(true);
+	Acceptor.set_option(reuse_option);
+	
+	tcp::endpoint Endpoint(tcp::v4(), uint16(Port));
+	Acceptor.bind(Endpoint);
+	Acceptor.listen(MaxQueuedConnections);
+
+	if (Acceptor.local_endpoint().port() != Port)
+	{
+		UE_LOG(LogCore, Error, TEXT("Failed to bind socket for port %u."), Port);
+		return false;
+	}
+
+	AsioContext->Acceptor = MoveTemp(Acceptor);
+	return true;
+}
+
+void FDirectSocketStream::Accept()
+{
+	using asio::ip::tcp;
+	if (!AsioContext.IsValid())
+	{
+		return;
+	}
+	
+	tcp::socket Socket(AsioContext->Context);
+	asio::error_code ErrorCode;
+	AsioContext->Acceptor.accept(Socket, ErrorCode);
+
+	if (ErrorCode)
+	{
+		if (ErrorCode == asio::error::interrupted)
+		{
+			return;
+		}
+		UE_LOG(LogCore, Error, TEXT("Failed accept sockets connection, error: %u."), ErrorCode.value());
+		return;
+	}
+	
+#if PLATFORM_WINDOWS
+	SetHandleInformation(HANDLE(SOCKET(Socket.native_handle())), HANDLE_FLAG_INHERIT, 0);
+#endif
+	
+	InternalStream = MakeUnique<FTraceDataStream>(Socket);
+	ConnectionEvent->Trigger();
+}
+
+uint32 FDirectSocketStream::Run()
+{
+	// Listening thread entry point. We currently only accept the first connection 
+	// so the thread will exit immediately after.
+	Accept();
+	return 0;
+}
+
+void FDirectSocketStream::Stop()
+{
+	if (AsioContext)
+	{
+		// Stopping the listening thread by closing the acceptor, which
+		// will abort the accept operation.
+		AsioContext->Acceptor.close();
+	}
+}
+
+bool FDirectSocketStream::WaitUntilReady()
+{
+	return ConnectionEvent->Wait(0, true);
+}
+
+int32 FDirectSocketStream::Read(void* Data, uint32 Size)
+{
+	if (!InternalStream)
+	{
+		// Treat trying to read the stream before it's ready as an error
+		return -1;
+	}
+	return InternalStream->Read(Data, Size);
+}
+
+void FDirectSocketStream::Close()
+{
+	if (InternalStream)
+	{
+		InternalStream->Close();
+	}
+}
+
+	
 } // namespace Trace
 } // namespace UE
diff --git a/Engine/Source/Developer/TraceAnalysis/Private/DataStreamInternal.h b/Engine/Source/Developer/TraceAnalysis/Private/DataStreamInternal.h
new file mode 100644
--- /dev/null
+++ b/Engine/Source/Developer/TraceAnalysis/Private/DataStreamInternal.h
@@ -0,0 +1,27 @@
+﻿// Copyright Epic Games, Inc. All Rights Reserved.
+#pragma once
+
+#include "Trace/DataStream.h"
+#include "Asio/Asio.h"
+
+namespace UE {
+namespace Trace {
+	
+////////////////////////////////////////////////////////////////////////////////
+class FTraceDataStream
+	: public IInDataStream
+{
+public:
+	FTraceDataStream(asio::ip::tcp::socket& InSocket);
+	virtual					~FTraceDataStream();
+	bool					IsOpen() const;
+	virtual void			Close() override;
+	virtual int32			Read(void* Dest, uint32 DestSize) override;
+
+private:
+	asio::ip::tcp::socket	Socket;
+};
+
+	
+} // namespace Trace
+} // namespace UE
\ No newline at end of file
diff --git a/Engine/Source/Developer/TraceAnalysis/Private/Store/StoreClient.cpp b/Engine/Source/Developer/TraceAnalysis/Private/Store/StoreClient.cpp
--- a/Engine/Source/Developer/TraceAnalysis/Private/Store/StoreClient.cpp
+++ b/Engine/Source/Developer/TraceAnalysis/Private/Store/StoreClient.cpp
@@ -4,96 +4,12 @@
 #include "Algo/Transform.h"
 #include "Asio/Asio.h"
 #include "CborPayload.h"
+#include "DataStreamInternal.h"
 #include "Templates/UnrealTemplate.h"
 
 namespace UE {
 namespace Trace {
 
-////////////////////////////////////////////////////////////////////////////////
-class FTraceDataStream
-	: public IInDataStream
-{
-public:
-							FTraceDataStream(asio::ip::tcp::socket& InSocket);
-	virtual					~FTraceDataStream();
-	bool					IsOpen() const;
-	virtual void			Close() override;
-	virtual int32			Read(void* Dest, uint32 DestSize) override;
-
-private:
-	asio::ip::tcp::socket	Socket;
-};
-
-////////////////////////////////////////////////////////////////////////////////
-FTraceDataStream::FTraceDataStream(asio::ip::tcp::socket& InSocket)
-: Socket(MoveTemp(InSocket))
-{
-	asio::socket_base::receive_buffer_size RecvBufferSize(4 << 20);
-	Socket.set_option(RecvBufferSize);
-}
-
-////////////////////////////////////////////////////////////////////////////////
-FTraceDataStream::~FTraceDataStream()
-{
-	Close();
-}
-
-////////////////////////////////////////////////////////////////////////////////
-bool FTraceDataStream::IsOpen() const
-{
-	return Socket.is_open();
-}
-
-////////////////////////////////////////////////////////////////////////////////
-void FTraceDataStream::Close()
-{
-	Socket.shutdown(asio::ip::tcp::socket::shutdown_receive);
-	Socket.close();
-}
-
-////////////////////////////////////////////////////////////////////////////////
-int32 FTraceDataStream::Read(void* Dest, uint32 DestSize)
-{
-	auto Handle = Socket.native_handle();
-
-	fd_set Fds;
-	FD_ZERO(&Fds);
-	FD_SET(Handle, &Fds);
-
-	timeval Timeout;
-	Timeout.tv_sec = 1;
-	Timeout.tv_usec = 0;
-
-	while (true)
-	{
-		fd_set ReadFds = Fds;
-		int Ret = select((int)Handle + 1, &ReadFds, 0, 0, &Timeout);
-
-		if (Ret < 0)
-		{
-			Close();
-			return -1;
-		}
-
-		if (Ret == 0)
-		{
-			continue;
-		}
-
-		asio::error_code ErrorCode;
-		size_t BytesRead = Socket.read_some(asio::buffer(Dest, DestSize), ErrorCode);
-		if (ErrorCode)
-		{
-			Close();
-			return -1;
-		}
-
-		return int32(BytesRead);
-	}
-}
-
-
-
 ////////////////////////////////////////////////////////////////////////////////
 class FStoreCborClient
 {
diff --git a/Engine/Source/Developer/TraceAnalysis/Public/Trace/DataStream.h b/Engine/Source/Developer/TraceAnalysis/Public/Trace/DataStream.h
--- a/Engine/Source/Developer/TraceAnalysis/Public/Trace/DataStream.h
+++ b/Engine/Source/Developer/TraceAnalysis/Public/Trace/DataStream.h
@@ -4,6 +4,7 @@
 
 #include "CoreFwd.h"
 #include "HAL/Platform.h"
+#include "HAL/Runnable.h"
 #include "Templates/UniquePtr.h"
 
 class IFileHandle;
@@ -15,8 +16,29 @@ class IInDataStream
 {
 public:
 	virtual			~IInDataStream() = default;
+
+	/**
+	 * Read bytes from the stream.
+	 * @param Data Pointer to buffer to read into
+	 * @param Size Maximum size that can be read in bytes
+	 * @return Number of bytes read from the stream. Zero indicates end of stream and negative values indicate errors.
+	 */
 	virtual int32	Read(void* Data, uint32 Size) = 0;
+
+	/**
+	 * Close the stream. Reading from a closed stream
+	 * is considered an error.
+	 */
 	virtual void	Close() {}
+	
+	/**
+	 * Query if the stream is ready to read. Some streams may need to
+	 * establish the data stream before reading can begin. A stream may not
+	 * block indefinitely.
+	 * 
+	 * @return if the stream is ready to be read from
+	 */
+	virtual bool	WaitUntilReady() { return true; }
 };
 
 /*
@@ -26,7 +48,7 @@ class TRACEANALYSIS_API FFileDataStream : public IInDataStream
 {
 public:
 	FFileDataStream();
-	~FFileDataStream();
+	virtual ~FFileDataStream() override;
 
 	/*
 	* Open the file.
@@ -37,19 +59,7 @@ public:
 	*/
 	bool Open(const TCHAR* Path);
 
-	/*
-	* Read from the file.
-	* 
-	* @param Data	The address in memory to write file data to.
-	* @param Size	The size of available memory at Data.
-	* 
-	* @return The number of bytes read from the file to Data, or zero in the case of a read failure.
-	*/
 	virtual int32 Read(void* Data, uint32 Size) override;
-
-	/*
-	* Close the file.
-	*/
 	virtual void Close() override;
 
 private:
@@ -57,5 +67,49 @@ private:
 	uint64 Remaining;
 };
 
+/**
+ * Creates a stream to directly consume a trace stream from the tracing application. Sets up
+ * a listening socket and stream is not considered ready until a connection is made.
+ */
+class TRACEANALYSIS_API FDirectSocketStream : public IInDataStream, public FRunnable
+{
+public:
+	FDirectSocketStream();
+	virtual ~FDirectSocketStream() override;
+
+	/**
+	 * Initiates listening sockets. Must be called before attempting to read from
+	 * the stream.
+	 * @return Port number used for listening
+	 */
+	uint16 StartListening();
+
+private:
+	enum
+	{
+		DefaultPort = 1986,			// Default port to use.
+		MaxPortAttempts = 16,		// How many port increments are tried if fail to bind default.
+		MaxQueuedConnections = 4,	// Size of connection queue.
+	};
+
+	// IInStream interface
+	virtual bool WaitUntilReady() override;
+	virtual int32 Read(void* Data, uint32 Size) override;
+	virtual void Close() override;
+
+	// FRunnable interface
+	virtual uint32 Run() override;
+	virtual void Stop() override;
+	
+	void Accept();
+	bool CreateSocket(uint16 Port);
+
+private:
+	TUniquePtr<struct FDirectSocketContext> AsioContext;
+	TUniquePtr<class FTraceDataStream> InternalStream;
+	TUniquePtr<class FRunnableThread> ListeningThread;
+	FEvent* ConnectionEvent;
+};
+
 } // namespace Trace
 } // namespace UE
