Log direct connection success

#rnx

diff --git a/Engine/Source/Developer/TraceAnalysis/Private/DataStream.cpp b/Engine/Source/Developer/TraceAnalysis/Private/DataStream.cpp
--- a/Engine/Source/Developer/TraceAnalysis/Private/DataStream.cpp
+++ b/Engine/Source/Developer/TraceAnalysis/Private/DataStream.cpp
@@ -258,6 +258,13 @@ void FDirectSocketStream::Accept()
 	SetHandleInformation(HANDLE(SOCKET(Socket.native_handle())), HANDLE_FLAG_INHERIT, 0);
 #endif
 	
+	{
+		const auto Address = Socket.remote_endpoint().address();
+		TStringBuilder<16> RemoteAddress;
+		RemoteAddress << ANSI_TO_TCHAR(Address.to_string().c_str());
+		UE_LOG(LogCore, Log, TEXT("Accepted direct trace connection from %s"), RemoteAddress.ToString());
+	}
+	
 	InternalStream = MakeUnique<FTraceDataStream>(Socket);
 	ConnectionEvent->Trigger();
 }
