Some hardening of direct socket stream trace connections

* Better error handling when setting up the acceptor.
* Make sure to close the acceptor once a connection is made.

#rnx

diff --git a/Engine/Source/Developer/TraceAnalysis/Private/DataStream.cpp b/Engine/Source/Developer/TraceAnalysis/Private/DataStream.cpp
--- a/Engine/Source/Developer/TraceAnalysis/Private/DataStream.cpp
+++ b/Engine/Source/Developer/TraceAnalysis/Private/DataStream.cpp
@@ -204,27 +204,40 @@ bool FDirectSocketStream::CreateSocket(uint16 Port)
 {
 	using asio::ip::tcp;
 
+	asio::error_code ErrorCode;
 	tcp::acceptor Acceptor(AsioContext->Context);
 
 #if PLATFORM_WINDOWS
 	DWORD Flags = WSA_FLAG_NO_HANDLE_INHERIT|WSA_FLAG_OVERLAPPED;
 	SOCKET Socket = WSASocketW(AF_INET, SOCK_STREAM, IPPROTO_TCP, nullptr, 0, Flags);
-	Acceptor.assign(tcp::v4(), Socket);
+	Acceptor.assign(tcp::v4(), Socket, ErrorCode);
 #else
-	Acceptor.open(tcp::v4());
+	Acceptor.open(tcp::v4(), ErrorCode);
 #endif
+	if (ErrorCode)
+	{
+		UE_LOG(LogCore, Error, TEXT("Failed to open acceptor: '%s'"), ANSI_TO_TCHAR(ErrorCode.message().c_str()));
+		return false;
+	}
 
 	// Reuse address to avoid locking up port
 	asio::socket_base::reuse_address reuse_option(true);
 	Acceptor.set_option(reuse_option);
 	
+	// Setup endpoint for port on any address
 	tcp::endpoint Endpoint(tcp::v4(), uint16(Port));
-	Acceptor.bind(Endpoint);
-	Acceptor.listen(MaxQueuedConnections);
 
-	if (Acceptor.local_endpoint().port() != Port)
+	Acceptor.bind(Endpoint, ErrorCode);
+	if (ErrorCode)
+	{
+		UE_LOG(LogCore, Error, TEXT("Failed to bind socket on port %u: '%s'"), Port, ANSI_TO_TCHAR(ErrorCode.message().c_str()));
+		return false;
+	}
+
+	Acceptor.listen(MaxQueuedConnections, ErrorCode);
+	if (ErrorCode)
 	{
-		UE_LOG(LogCore, Error, TEXT("Failed to bind socket for port %u."), Port);
+		UE_LOG(LogCore, Error, TEXT("Failed to listen on port %u: '%s'"), Port, ANSI_TO_TCHAR(ErrorCode.message().c_str()));
 		return false;
 	}
 
@@ -265,6 +278,7 @@ void FDirectSocketStream::Accept()
 		UE_LOG(LogCore, Log, TEXT("Accepted direct trace connection from %s"), RemoteAddress.ToString());
 	}
 	
+	AsioContext->Acceptor.close();
 	InternalStream = MakeUnique<FTraceDataStream>(Socket);
 	ConnectionEvent->Trigger();
 }
