TraceAnalysis: No longer needs to force c++17

#rnx

diff --git a/Engine/Source/Developer/TraceAnalysis/TraceAnalysis.Build.cs b/Engine/Source/Developer/TraceAnalysis/TraceAnalysis.Build.cs
--- a/Engine/Source/Developer/TraceAnalysis/TraceAnalysis.Build.cs
+++ b/Engine/Source/Developer/TraceAnalysis/TraceAnalysis.Build.cs
@@ -8,9 +8,6 @@ public class TraceAnalysis : ModuleRules
 	{
 		CppCompileWarningSettings.UnsafeTypeCastWarningLevel = WarningLevel.Error;
 
-		// Asio does not compile with C++20, remove if updated
-		CppStandard = CppStandardVersion.Cpp17;
-
 		PrivateDependencyModuleNames.AddRange(
 			new string[] {
 				"Asio",
