[Trace] Fixed compilation for TraceAnalyzer.exe.

#rnx

diff --git a/Engine/Source/Developer/TraceAnalysis/Public/TraceAnalysisDebug.h b/Engine/Source/Developer/TraceAnalysis/Public/TraceAnalysisDebug.h
--- a/Engine/Source/Developer/TraceAnalysis/Public/TraceAnalysisDebug.h
+++ b/Engine/Source/Developer/TraceAnalysis/Public/TraceAnalysisDebug.h
@@ -30,17 +30,17 @@
 
 #if UE_TRACE_ANALYSIS_DEBUG_LOG_IMPL == 1 // only printf
 
-	#define UE_TRACE_ANALYSIS_DEBUG_LOG(format, ...)        printf(format "\n", __VA_ARGS__)
+	#define UE_TRACE_ANALYSIS_DEBUG_LOG(format, ...)        printf(format "\n", ##__VA_ARGS__)
 	#define UE_TRACE_ANALYSIS_DEBUG_BeginStringBuilder()
 	#define UE_TRACE_ANALYSIS_DEBUG_Append(s)               printf("%s", s)
-	#define UE_TRACE_ANALYSIS_DEBUG_Appendf(format, ...)    printf(format, __VA_ARGS__)
+	#define UE_TRACE_ANALYSIS_DEBUG_Appendf(format, ...)    printf(format, ##__VA_ARGS__)
 	#define UE_TRACE_ANALYSIS_DEBUG_AppendChar(c)           printf("%c", c)
 	#define UE_TRACE_ANALYSIS_DEBUG_EndStringBuilder()      printf("\n")
 	#define UE_TRACE_ANALYSIS_DEBUG_ResetStringBuilder()
 
 #elif UE_TRACE_ANALYSIS_DEBUG_LOG_IMPL == 2 // printf with TAnsiStringBuilder
 
-	#define UE_TRACE_ANALYSIS_DEBUG_LOG(format, ...)        printf(format "\n", __VA_ARGS__)
+	#define UE_TRACE_ANALYSIS_DEBUG_LOG(format, ...)        printf(format "\n", ##__VA_ARGS__)
 	#define UE_TRACE_ANALYSIS_DEBUG_BeginStringBuilder()    TAnsiStringBuilder<1024> StringBuilder;
 	#define UE_TRACE_ANALYSIS_DEBUG_Append                  StringBuilder.Append
 	#define UE_TRACE_ANALYSIS_DEBUG_Appendf                 StringBuilder.Appendf
@@ -55,7 +55,7 @@
 	#include "Misc/OutputDeviceRedirector.h"
 	#include "Misc/StringBuilder.h"
 
-	#define UE_TRACE_ANALYSIS_DEBUG_LOG(format, ...)        { TAnsiStringBuilder<1024> SB; SB.Appendf(format, __VA_ARGS__); GLog->Logf(TEXT("%hs"), SB.ToString()); }
+	#define UE_TRACE_ANALYSIS_DEBUG_LOG(format, ...)        { TAnsiStringBuilder<1024> SB; SB.Appendf(format, ##__VA_ARGS__); GLog->Logf(TEXT("%hs"), SB.ToString()); }
 	#define UE_TRACE_ANALYSIS_DEBUG_BeginStringBuilder()    TAnsiStringBuilder<1024> StringBuilder;
 	#define UE_TRACE_ANALYSIS_DEBUG_Append                  StringBuilder.Append
 	#define UE_TRACE_ANALYSIS_DEBUG_Appendf                 StringBuilder.Appendf
@@ -67,10 +67,10 @@
 
 	#include "Misc/StringBuilder.h"
 
-	#define UE_TRACE_ANALYSIS_DEBUG_LOG(format, ...)        { TAnsiStringBuilder<1024> SB; SB.Appendf(format, __VA_ARGS__); FPlatformMisc::LowLevelOutputDebugStringf(TEXT("%hs\n"), SB.ToString()); }
+	#define UE_TRACE_ANALYSIS_DEBUG_LOG(format, ...)        { TAnsiStringBuilder<1024> SB; SB.Appendf(format, ##__VA_ARGS__); FPlatformMisc::LowLevelOutputDebugStringf(TEXT("%hs\n"), SB.ToString()); }
 	#define UE_TRACE_ANALYSIS_DEBUG_BeginStringBuilder()
 	#define UE_TRACE_ANALYSIS_DEBUG_Append(s)               FPlatformMisc::LowLevelOutputDebugStringf(TEXT("%hs"), s)
-	#define UE_TRACE_ANALYSIS_DEBUG_Appendf(format, ...)    { TAnsiStringBuilder<1024> SB; SB.Appendf(format, __VA_ARGS__); FPlatformMisc::LowLevelOutputDebugStringf(TEXT("%hs"), SB.ToString()); }
+	#define UE_TRACE_ANALYSIS_DEBUG_Appendf(format, ...)    { TAnsiStringBuilder<1024> SB; SB.Appendf(format, ##__VA_ARGS__); FPlatformMisc::LowLevelOutputDebugStringf(TEXT("%hs"), SB.ToString()); }
 	#define UE_TRACE_ANALYSIS_DEBUG_AppendChar(c)           FPlatformMisc::LowLevelOutputDebugStringf(TEXT("%hc"), c)
 	#define UE_TRACE_ANALYSIS_DEBUG_EndStringBuilder()      FPlatformMisc::LowLevelOutputDebugString(TEXT("\n"))
 	#define UE_TRACE_ANALYSIS_DEBUG_ResetStringBuilder()
