[Insights] * TraceAnalysis: Fixed a huge memory leak regarding trace analysis. The entire Transport object for the TidPacket(Sync) transport protocols (and all its buffers used to decode/analyze the utrace stream) was leaking, on each analysis. The issue affected analysis of utrace streams with version 5 and newer (i.e. issue introduced with FProtocol5Stage).
* Fixed a rare issue that could occur when dispatching events (the min-heap not being updated correctly after removing one stream).
* Increased UE_TRACE_ANALYSIS_DEBUG_LEVEL for some verbose debug log messages.

diff --git a/Engine/Source/Developer/TraceAnalysis/Private/Analysis/Engine.cpp b/Engine/Source/Developer/TraceAnalysis/Private/Analysis/Engine.cpp
--- a/Engine/Source/Developer/TraceAnalysis/Private/Analysis/Engine.cpp
+++ b/Engine/Source/Developer/TraceAnalysis/Private/Analysis/Engine.cpp
@@ -2627,7 +2627,8 @@ class FProtocol0Stage
 {
 public:
 						FProtocol0Stage(FTransport* InTransport);
-						~FProtocol0Stage();
+	virtual				~FProtocol0Stage();
+
 	virtual EStatus		OnData(FStreamReader& Reader, const FMachineContext& Context) override;
 	virtual void		EnterStage(const FMachineContext& Context) override;
 	virtual void		ExitStage(const FMachineContext& Context) override;
@@ -2742,7 +2743,8 @@ class FProtocol2Stage
 {
 public:
 						FProtocol2Stage(uint32 Version, FTransport* InTransport);
-						~FProtocol2Stage();
+	virtual				~FProtocol2Stage();
+
 	virtual EStatus		OnData(FStreamReader& Reader, const FMachineContext& Context) override;
 	virtual void		EnterStage(const FMachineContext& Context) override;
 	virtual void		ExitStage(const FMachineContext& Context) override;
@@ -2919,8 +2921,10 @@ FProtocol2Stage::EStatus FProtocol2Stage::OnData(
 			break;
 		}
 
-		TArrayView<FRotaItem> ActiveRota(Rota.GetData(), ActiveCount);
-		Algo::Sort(ActiveRota);
+		{
+			TArrayView<FRotaItem> ActiveRota(Rota.GetData(), ActiveCount);
+			Algo::Sort(ActiveRota);
+		}
 
 		int32 MinLogSerial = Rota[0].Serial;
 		if (ActiveCount > 1)
@@ -3400,9 +3404,15 @@ int32 FProtocol4Stage::OnDataKnown(
 	case Protocol4::EKnownEventUids::NewEvent:
 		{
 			const auto* Size = Reader.GetPointer<uint16>();
-			check(Size != nullptr);
+			if (!Size)
+			{
+				return 0;
+			}
 			const void* EventTypeAndData = Reader.GetPointer(sizeof(*Size) + *Size);
-			check(EventTypeAndData != nullptr);
+			if (!EventTypeAndData)
+			{
+				return 0;
+			}
 			const FTypeRegistry::FTypeInfo* TypeInfo = TypeRegistry.Add(Size + 1, 4);
 #if UE_TRACE_ANALYSIS_DEBUG
 			Bridge.DebugLogNewEvent(Uid, TypeInfo, uint32(State.TempUidBytes + sizeof(*Size) + *Size));
@@ -3429,7 +3439,7 @@ int32 FProtocol4Stage::OnDataKnown(
 	case Protocol4::EKnownEventUids::EnterScope_T:
 		{
 			const uint8* Stamp = Reader.GetPointer(sizeof(uint64) - 1);
-			if (Stamp == nullptr)
+			if (!Stamp)
 			{
 				return 0;
 			}
@@ -3447,7 +3457,7 @@ int32 FProtocol4Stage::OnDataKnown(
 	case Protocol4::EKnownEventUids::LeaveScope_T:
 		{
 			const uint8* Stamp = Reader.GetPointer(sizeof(uint64) - 1);
-			if (Stamp == nullptr)
+			if (!Stamp)
 			{
 				return 0;
 			}
@@ -3478,6 +3488,8 @@ class FProtocol5Stage
 {
 public:
 						FProtocol5Stage(FTransport* InTransport);
+	virtual				~FProtocol5Stage();
+
 	virtual void		ExitStage(const FMachineContext& Context) override;
 
 protected:
@@ -3574,6 +3586,7 @@ protected:
 #endif // UE_TRACE_ANALYSIS_DEBUG
 
 	FTypeRegistry			TypeRegistry;
+	FTransport*				BaseTransport;
 	FTidPacketTransport&	Transport;
 	EventDescArray			EventDescs;
 	EventDescArray			SerialGaps;
@@ -3588,12 +3601,19 @@ protected:
 
 ////////////////////////////////////////////////////////////////////////////////
 FProtocol5Stage::FProtocol5Stage(FTransport* InTransport)
-: Transport(*(FTidPacketTransport*)InTransport)
+: BaseTransport(InTransport)
+, Transport(*(FTidPacketTransport*)InTransport)
 , SyncCount(Transport.GetSyncCount())
 {
 	EventDescs.Reserve(8 << 10);
 }
 
+////////////////////////////////////////////////////////////////////////////////
+FProtocol5Stage::~FProtocol5Stage()
+{
+	delete BaseTransport;
+}
+
 ////////////////////////////////////////////////////////////////////////////////
 void FProtocol5Stage::ExitStage(const FMachineContext& Context)
 {
@@ -3958,9 +3978,25 @@ FProtocol5Stage::EStatus FProtocol5Stage::OnDataNormal(const FMachineContext& Co
 		NextSerialWaitCount > MaxNextSerialWaitCount)
 	{
 		UE_TRACE_ANALYSIS_DEBUG_LOG("Error: Trace analysis accumulated too much data (%d parsed events) and will start to skip the missing serial sync events!", NumAvailableEvents);
+		UE_TRACE_ANALYSIS_DEBUG_LOG("NextSerial=%u", NextSerial);
+
 		if (!bSkipSerialError)
 		{
 			bSkipSerialError = true;
+#if UE_TRACE_ANALYSIS_DEBUG && UE_TRACE_ANALYSIS_DEBUG_LEVEL >= 2
+			for (const auto& Stream : EventDescHeap)
+			{
+				const FEventDesc* EventDesc = Stream.EventDescs;
+				while (EventDesc->Serial != ESerial::Terminal)
+				{
+					if (EventDesc->Serial == NextSerial)
+					{
+						UE_TRACE_ANALYSIS_DEBUG_LOG("Found next serial on thread %u (event %d)", Stream.ThreadId, (int32)(EventDesc - Stream.EventDescs));
+					}
+					++EventDesc;
+				}
+			}
+#endif
 			Context.EmitMessagef(
 				EAnalysisMessageSeverity::Error,
 				TEXT("Trace analysis accumulated too much data (%d parsed events) and will start to skip the missing serial sync events!"), NumAvailableEvents);
@@ -4148,9 +4184,7 @@ int32 FProtocol5Stage::DispatchNormalEvents(const FMachineContext& Context, TArr
 }
 
 ////////////////////////////////////////////////////////////////////////////////
-int32 FProtocol5Stage::DispatchEvents(
-	const FMachineContext& Context,
-	TArray<FEventDescStream>& EventDescHeap)
+int32 FProtocol5Stage::DispatchEvents(const FMachineContext& Context, TArray<FEventDescStream>& EventDescHeap)
 {
 #if UE_TRACE_ANALYSIS_DEBUG
 	FAnalysisState& State = Context.Bridge.GetState();
@@ -4160,12 +4194,17 @@ int32 FProtocol5Stage::DispatchEvents(
 	{
 		if (EventDesc->Serial != ESerial::Terminal)
 		{
+			// Stream reference is for an element of EventDescHeap array, so we copy it before modifying the array.
 			FEventDescStream Next = Stream;
 			Next.EventDescs = EventDesc;
-			EventDescHeap.Add(Next);
+			// We can now discard the top stream from the heap and push the updated stream to the heap.
+			EventDescHeap.HeapPopDiscard(FSerialDistancePredicate{ NextSerial }, EAllowShrinking::No);
+			EventDescHeap.HeapPush(Next, FSerialDistancePredicate{ NextSerial });
+		}
+		else
+		{
+			EventDescHeap.HeapPopDiscard(FSerialDistancePredicate{ NextSerial }, EAllowShrinking::No);
 		}
-
-		EventDescHeap.HeapPopDiscard(FSerialDistancePredicate{NextSerial}, EAllowShrinking::No);
 	};
 
 	int32 NumDispatchedEvents = 0;
@@ -4192,7 +4231,7 @@ int32 FProtocol5Stage::DispatchEvents(
 			UE_TRACE_ANALYSIS_DEBUG_LOG("Skip serial gap (%u +%u) --> NextSerial=%u", EndDesc->Serial, EndDesc->GapLength, NextSerial);
 #endif
 			UpdateHeap(Stream, EndDesc + 1);
-			NumDispatchedEvents += ((EndDesc + 1)->Serial == ESerial::Terminal) ? 2 : 1;// GapThreadId event + Terminal event
+			NumDispatchedEvents += ((EndDesc + 1)->Serial == ESerial::Terminal) ? 2 : 1; // GapThreadId event + Terminal event
 			continue;
 		}
 
@@ -4244,11 +4283,10 @@ int32 FProtocol5Stage::DispatchEvents(
 				uint32 BufferSize = 0;
 				uint32 DataSize = 0; // parsed data size + remaining data size
 				uint32 RemainingSize = 0;
-				FTidPacketTransport* InnerTransport = (FTidPacketTransport*)(&Transport);
-				for (uint32 i = 0, n = InnerTransport->GetThreadCount(); i < n; ++i)
+				for (uint32 i = 0, n = Transport.GetThreadCount(); i < n; ++i)
 				{
-					FStreamBuffer* ThreadReader = (FStreamBuffer*)InnerTransport->GetThreadStream(i);
-					uint32 ThreadId = InnerTransport->GetThreadId(i);
+					FStreamBuffer* ThreadReader = (FStreamBuffer*)Transport.GetThreadStream(i);
+					uint32 ThreadId = Transport.GetThreadId(i);
 					if (ThreadId == EventDescStream.ThreadId)
 					{
 						BufferSize = ThreadReader->GetBufferSize();
@@ -4489,7 +4527,7 @@ int32 FProtocol5Stage::ParseEventsWithAux(FStreamReader& Reader, EventDescArray&
 	if (bUnsorted)
 	{
 		uint32 NumDescs = OutEventDescs.Num() - FirstDescIndex;
-#if UE_TRACE_ANALYSIS_DEBUG && UE_TRACE_ANALYSIS_DEBUG_LEVEL >= 2
+#if UE_TRACE_ANALYSIS_DEBUG && UE_TRACE_ANALYSIS_DEBUG_LEVEL >= 3
 		UE_TRACE_ANALYSIS_DEBUG_LOG("Sorting %u event descs", NumDescs);
 #endif
 		TArrayView<FEventDesc> DescsView(OutEventDescs.GetData() + FirstDescIndex, NumDescs);
@@ -4590,9 +4628,7 @@ int32 FProtocol5Stage::ParseEvent(FStreamReader& Reader, FEventDesc& EventDesc,
 
 ////////////////////////////////////////////////////////////////////////////////
 template <typename Callback>
-void FProtocol5Stage::ForEachSerialGap(
-	const TArray<FEventDescStream>& EventDescHeap,
-	Callback&& InCallback)
+void FProtocol5Stage::ForEachSerialGap(const TArray<FEventDescStream>& EventDescHeap, Callback&& InCallback)
 {
 	TArray<FEventDescStream> HeapCopy(EventDescHeap);
 	int Serial = HeapCopy.HeapTop().EventDescs[0].Serial;
@@ -4817,15 +4853,13 @@ bool FProtocol5Stage::DispatchKnownEvent(const FMachineContext& Context, uint32
 }
 
 ////////////////////////////////////////////////////////////////////////////////
-int32 FProtocol5Stage::DispatchEvents(
-	const FMachineContext& Context,
-	const FEventDesc* EventDesc, uint32 Count)
+int32 FProtocol5Stage::DispatchEvents(const FMachineContext& Context, const FEventDesc* EventDesc, uint32 Count)
 {
 	using namespace Protocol5;
 
 	FAuxDataCollector AuxCollector;
 
-#if UE_TRACE_ANALYSIS_DEBUG && UE_TRACE_ANALYSIS_DEBUG_LEVEL >= 2
+#if UE_TRACE_ANALYSIS_DEBUG && UE_TRACE_ANALYSIS_DEBUG_LEVEL >= 3
 	UE_TRACE_ANALYSIS_DEBUG_LOG("Dispatch run of %u consecutive events (Tid=%u)", Count, Context.Bridge.GetActiveThreadId());
 #endif
 
@@ -4931,6 +4965,7 @@ class FProtocol6Stage
 {
 public:
 						FProtocol6Stage(FTransport* InTransport);
+
 	virtual void		ExitStage(const FMachineContext& Context) override;
 };
 
@@ -4966,6 +5001,7 @@ class FProtocol7Stage
 {
 public:
 						FProtocol7Stage(FTransport* InTransport);
+
 	virtual void		ExitStage(const FMachineContext& Context) override;
 	virtual void		SetSizeIfKnownEvent(uint32 Uid, uint32& InOutEventSize) override;
 	virtual bool		DispatchKnownEvent(const FMachineContext& Context, uint32 Uid, const FEventDesc* Cursor) override;
diff --git a/Engine/Source/Developer/TraceAnalysis/Private/Analysis/Transport/PacketTransport.h b/Engine/Source/Developer/TraceAnalysis/Private/Analysis/Transport/PacketTransport.h
--- a/Engine/Source/Developer/TraceAnalysis/Private/Analysis/Transport/PacketTransport.h
+++ b/Engine/Source/Developer/TraceAnalysis/Private/Analysis/Transport/PacketTransport.h
@@ -12,6 +12,7 @@ class FPacketTransport
 	: public FTransport
 {
 public:
+							FPacketTransport() {}
 	virtual					~FPacketTransport();
 	virtual void			Advance(uint32 BlockSize) override;
 	virtual bool			IsEmpty() const override;
diff --git a/Engine/Source/Developer/TraceAnalysis/Private/Analysis/Transport/TidPacketTransport.h b/Engine/Source/Developer/TraceAnalysis/Private/Analysis/Transport/TidPacketTransport.h
--- a/Engine/Source/Developer/TraceAnalysis/Private/Analysis/Transport/TidPacketTransport.h
+++ b/Engine/Source/Developer/TraceAnalysis/Private/Analysis/Transport/TidPacketTransport.h
@@ -19,14 +19,15 @@ class FTidPacketTransport
 	: public FTransport
 {
 public:
-	typedef UPTRINT ThreadIter;
-	
 	enum class ETransportResult : uint8
 	{
 		Ok,
 		Error
 	};
 
+							FTidPacketTransport() {}
+	virtual					~FTidPacketTransport() {}
+
 	virtual bool			IsEmpty() const override;
 	virtual void			DebugBegin() override;
 	virtual void			DebugEnd() override;
@@ -37,8 +38,14 @@ public:
 	uint32					GetThreadId(uint32 Index) const;
 	uint32					GetSyncCount() const;
 
+protected:
+	uint32					Synced = 0x7fff'ffff;
+	uint32					NumPackets = 0;
+#if UE_TRACE_PACKET_VERIFICATION
+	uint64					LastPacketSerial = 0;
+#endif
+
 private:
-	
 	struct FThreadStream
 	{
 		FStreamBuffer		Buffer;
@@ -61,18 +68,10 @@ private:
 								{ {}, ETransportTid::Importants },
 							};
 
-protected:
-	uint32					Synced = 0x7fff'ffff;
-#if UE_TRACE_PACKET_VERIFICATION
-	uint64					LastPacketSerial = 0;
-#endif
-
 #if UE_TRACE_ANALYSIS_DEBUG
-private:
 	uint64					TotalPacketHeaderSize = 0;
 	uint64					TotalPacketSize = 0;
 	uint64					TotalDecodedSize = 0;
-	uint32					NumPackets = 0;
 	uint32					MaxBufferSize = 0;
 	uint32					MaxBufferSizeThreadId = 0;
 	uint32					MaxDataSizePerBuffer = 0;
@@ -98,6 +97,7 @@ public:
 	{
 		Synced = 0;
 	}
+	virtual ~FTidPacketTransportSync() {}
 };
 
 } // namespace Trace
diff --git a/Engine/Source/Developer/TraceAnalysis/Private/Analysis/Transport/Transport.h b/Engine/Source/Developer/TraceAnalysis/Private/Analysis/Transport/Transport.h
--- a/Engine/Source/Developer/TraceAnalysis/Private/Analysis/Transport/Transport.h
+++ b/Engine/Source/Developer/TraceAnalysis/Private/Analysis/Transport/Transport.h
@@ -11,6 +11,7 @@ namespace Trace {
 class FTransport
 {
 public:
+							FTransport() {}
 	virtual					~FTransport() {}
 	void					SetReader(FStreamReader& InReader);
 	FStreamReader*			GetReader() const { return Reader; }
@@ -26,7 +27,7 @@ public:
 protected:
 	virtual const uint8*	GetPointerImpl(uint32 BlockSize);
 
-	FStreamReader*			Reader;
+	FStreamReader*			Reader = nullptr;
 };
 
 ////////////////////////////////////////////////////////////////////////////////
