Avoid re-allocating the event descriptions array when we roll back an event due to an error.

diff --git a/Engine/Source/Developer/TraceAnalysis/Private/Analysis/Engine.cpp b/Engine/Source/Developer/TraceAnalysis/Private/Analysis/Engine.cpp
--- a/Engine/Source/Developer/TraceAnalysis/Private/Analysis/Engine.cpp
+++ b/Engine/Source/Developer/TraceAnalysis/Private/Analysis/Engine.cpp
@@ -4415,7 +4415,7 @@ int32 FProtocol5Stage::ParseEvents(FStreamReader& Reader, EventDescArray& OutEve
 #if UE_TRACE_ANALYSIS_DEBUG && UE_TRACE_ANALYSIS_DEBUG_LEVEL >= 2
 				UE_TRACE_ANALYSIS_DEBUG_LOG("Warning: Incomplete aux stack (Uid=%u)! Rewind %d parsed events.", EventDesc.Uid, OutEventDescs.Num() - RewindDescsNum);
 #endif
-				OutEventDescs.SetNum(RewindDescsNum);
+				OutEventDescs.SetNum(RewindDescsNum, EAllowShrinking::No);
 				Reader.RestoreMark(RewindMark);
 				break;
 			}
