Adding an FSocket::ReleaseNativeSocket for use with native socket API's (e.g. asio).  Using this new function inside the TraceAnalysis' FDirectSocketStream to support cases where we want multiple connections, not just a single connection (e.g. a multi-server capture).

diff --git a/Engine/Source/Developer/TraceAnalysis/Private/DataStream.cpp b/Engine/Source/Developer/TraceAnalysis/Private/DataStream.cpp
--- a/Engine/Source/Developer/TraceAnalysis/Private/DataStream.cpp
+++ b/Engine/Source/Developer/TraceAnalysis/Private/DataStream.cpp
@@ -10,6 +10,8 @@
 #include "HAL/Runnable.h"
 #include "HAL/RunnableThread.h"
 
+#include "Sockets.h"
+
 namespace UE {
 namespace Trace {
 
@@ -288,6 +290,31 @@ void FDirectSocketStream::Accept()
 	ConnectionEvent->Trigger();
 }
 
+bool FDirectSocketStream::BindToExistingConnection(FSocket&& ExistingConnectedSocket)
+{
+	using asio::ip::tcp;
+
+	// Take over the existing connection by releasing it
+	UPTRINT NativeSocket = ExistingConnectedSocket.ReleaseNativeSocket();
+	tcp::socket AsioSocket(AsioContext->Context);
+
+	// Assign it to our internal asio representation
+	asio::error_code ErrorCode;
+	AsioSocket.assign(tcp::v4(), NativeSocket, ErrorCode);
+
+	if (ErrorCode)
+	{
+		UE_LOG(LogCore, Error, TEXT("%hs: Failed to bind ASIO socket from FSocket connection, error: %u."), __func__, ErrorCode.value());
+		return false;
+	}
+
+	// mimic behavior that we just completed the connection
+	InternalStream = MakeUnique<FTraceDataStream>(AsioSocket);
+	ConnectionEvent->Trigger();
+
+	return true;
+}
+
 uint32 FDirectSocketStream::Run()
 {
 	// Listening thread entry point. We currently only accept the first connection 
@@ -329,6 +356,6 @@ void FDirectSocketStream::Close()
 	}
 }
 
-	
+
 } // namespace Trace
 } // namespace UE
diff --git a/Engine/Source/Developer/TraceAnalysis/Public/Trace/DataStream.h b/Engine/Source/Developer/TraceAnalysis/Public/Trace/DataStream.h
--- a/Engine/Source/Developer/TraceAnalysis/Public/Trace/DataStream.h
+++ b/Engine/Source/Developer/TraceAnalysis/Public/Trace/DataStream.h
@@ -10,6 +10,7 @@
 #define UE_API TRACEANALYSIS_API
 
 class IFileHandle;
+class FSocket;
 
 namespace UE {
 namespace Trace {
@@ -87,6 +88,15 @@ public:
 	 */
 	UE_API uint16 StartListening(uint16 Port = 0);
 
+	/**
+	 * Bind this instance to an existing connection that has already been established.
+	 * Using this method allows you to service connections off a single listen port, then 
+	 * create instances of FDirectSocketStream for each connection.
+	 * @param ExistingConnectedSocket is the existing socket that will transfer its ownership to this instance, even if the function fails.
+	 * @return true if the binding was successful. Note: Even if it returns false, you will have lost ownership of ExistingConnectedSocket.
+	 */
+	UE_API bool BindToExistingConnection(FSocket&& ExistingConnectedSocket);
+
 private:
 	enum
 	{
diff --git a/Engine/Source/Runtime/Sockets/Private/BSDSockets/SocketsBSD.cpp b/Engine/Source/Runtime/Sockets/Private/BSDSockets/SocketsBSD.cpp
--- a/Engine/Source/Runtime/Sockets/Private/BSDSockets/SocketsBSD.cpp
+++ b/Engine/Source/Runtime/Sockets/Private/BSDSockets/SocketsBSD.cpp
@@ -36,6 +36,14 @@ bool FSocketBSD::Shutdown(ESocketShutdownMode Mode)
 	return shutdown(Socket, InternalMode) == 0;
 }
 
+UPTRINT FSocketBSD::ReleaseNativeSocket()
+{
+	UPTRINT NativeSocket = Socket;
+	Socket = INVALID_SOCKET;
+
+	return NativeSocket;
+}
+
 bool FSocketBSD::Close(void)
 {
 	if (Socket != INVALID_SOCKET)
diff --git a/Engine/Source/Runtime/Sockets/Private/BSDSockets/SocketsBSD.h b/Engine/Source/Runtime/Sockets/Private/BSDSockets/SocketsBSD.h
--- a/Engine/Source/Runtime/Sockets/Private/BSDSockets/SocketsBSD.h
+++ b/Engine/Source/Runtime/Sockets/Private/BSDSockets/SocketsBSD.h
@@ -121,6 +121,7 @@ public:
 	virtual int32 GetPortNo() override;
 	virtual bool SetIpPktInfo(bool bEnable) override;
 	virtual bool RecvFromWithPktInfo(uint8* Data, int32 BufferSize, int32& BytesRead, FInternetAddr& Source, FInternetAddr& Destination, ESocketReceiveFlags::Type Flags = ESocketReceiveFlags::None) override;
+	[[nodiscard]] virtual UPTRINT ReleaseNativeSocket() override;
 
 	bool SetIPv6Only(bool bIPv6Only);
 
diff --git a/Engine/Source/Runtime/Sockets/Public/Sockets.h b/Engine/Source/Runtime/Sockets/Public/Sockets.h
--- a/Engine/Source/Runtime/Sockets/Public/Sockets.h
+++ b/Engine/Source/Runtime/Sockets/Public/Sockets.h
@@ -92,6 +92,14 @@ public:
 	 */
 	virtual bool Close() = 0;
 
+	/**
+	 * Releases the underlying socket implementation for use with native APIs.
+	 * The user of this class must know the underlying implementation details in order to cast the return value properly.
+	 * If this function succeeds, the FSocket instance no longer controls the connection (it will not Close()).
+	 * @return a non-zero value upon success, or zero if not implemented/failure to release ownership.
+	 */
+	[[nodiscard]] virtual UPTRINT ReleaseNativeSocket() { return 0; }
+
 	/**
 	 * Binds a socket to a network byte ordered address.
 	 *
