[DSTM] Added support to save trace data to disk in socket mode.

- Added a new commandline argument to specify the directory where the trace files should be created -dstmsavedir="fullpath" . If the directory does not exist, it will be created. If the path is invalid, the trace will continue but no data will be saved to disk. Lastly if the commandline argument is not provided, the trace files will be created in the user's document folder.

- Added a new alternative constructor to FDirectSocketStream that takes a file writer instance, which will be used to write the data to disk as it is being read from the socket

diff --git a/Engine/Source/Developer/TraceAnalysis/Private/DataStream.cpp b/Engine/Source/Developer/TraceAnalysis/Private/DataStream.cpp
--- a/Engine/Source/Developer/TraceAnalysis/Private/DataStream.cpp
+++ b/Engine/Source/Developer/TraceAnalysis/Private/DataStream.cpp
@@ -169,6 +169,13 @@ FDirectSocketStream::FDirectSocketStream()
 	AsioContext = MakeUnique<FDirectSocketContext>();
 }
 
+FDirectSocketStream::FDirectSocketStream(TUniquePtr<FArchiveFileWriterGeneric>&& InFileWriter)
+	: FileWriter(MoveTemp(InFileWriter)),
+	ConnectionEvent(FPlatformProcess::GetSynchEventFromPool(false))
+
+{
+	AsioContext = MakeUnique<FDirectSocketContext>();
+}
 
 FDirectSocketStream::~FDirectSocketStream()
 {
@@ -345,7 +352,15 @@ int32 FDirectSocketStream::Read(void* Data, uint32 Size)
 		// Treat trying to read the stream before it's ready as an error
 		return -1;
 	}
-	return InternalStream->Read(Data, Size);
+
+	const int32 BytesRead = InternalStream->Read(Data, Size);
+
+	if (BytesRead > 0 && FileWriter.IsValid() && !FileWriter->IsError())
+	{
+		FileWriter->Serialize(Data, BytesRead);
+	}
+
+	return BytesRead;
 }
 
 void FDirectSocketStream::Close()
@@ -354,6 +369,12 @@ void FDirectSocketStream::Close()
 	{
 		InternalStream->Close();
 	}
+
+	if (FileWriter.IsValid())
+    {
+    	FileWriter->Close();
+    	FileWriter.Reset();
+    }
 }
 
 
diff --git a/Engine/Source/Developer/TraceAnalysis/Public/Trace/DataStream.h b/Engine/Source/Developer/TraceAnalysis/Public/Trace/DataStream.h
--- a/Engine/Source/Developer/TraceAnalysis/Public/Trace/DataStream.h
+++ b/Engine/Source/Developer/TraceAnalysis/Public/Trace/DataStream.h
@@ -3,6 +3,7 @@
 #pragma once
 
 #include "CoreFwd.h"
+#include "HAL/FileManagerGeneric.h"
 #include "HAL/Platform.h"
 #include "HAL/Runnable.h"
 #include "Templates/UniquePtr.h"
@@ -73,11 +74,13 @@ private:
 /**
  * Creates a stream to directly consume a trace stream from the tracing application. Sets up
  * a listening socket and stream is not considered ready until a connection is made.
+ * If a valid file writer is provided during construction, the trace data will also be written to disk.
  */
 class FDirectSocketStream : public IInDataStream, public FRunnable
 {
 public:
 	UE_API FDirectSocketStream();
+	UE_API FDirectSocketStream(TUniquePtr<FArchiveFileWriterGeneric>&& InFileWriter);
 	UE_API virtual ~FDirectSocketStream() override;
 
 	/**
@@ -121,6 +124,7 @@ private:
 	TUniquePtr<struct FDirectSocketContext> AsioContext;
 	TUniquePtr<class FTraceDataStream> InternalStream;
 	TUniquePtr<class FRunnableThread> ListeningThread;
+	TUniquePtr<FArchiveFileWriterGeneric> FileWriter;
 	FEvent* ConnectionEvent;
 };
 
