Fix Clang 21 warnings -Wimplicit-int-conversion -Wimplicit-int-conversion-on-negation

#rnx

diff --git a/Engine/Source/Developer/TraceAnalysis/Private/Analysis/Engine.cpp b/Engine/Source/Developer/TraceAnalysis/Private/Analysis/Engine.cpp
--- a/Engine/Source/Developer/TraceAnalysis/Private/Analysis/Engine.cpp
+++ b/Engine/Source/Developer/TraceAnalysis/Private/Analysis/Engine.cpp
@@ -787,7 +787,7 @@ const void* IAnalyzer::FArrayReader::GetImpl(uint32 Index, int8& SizeAndType) co
 		return nullptr;
 	}
 
-	SizeAndType = (SizeAndType < 0) ? -SizeAndType : SizeAndType;
+	SizeAndType = (SizeAndType < 0) ? static_cast<int8>(-SizeAndType) : SizeAndType;
 	return Inner->Data + (Index * SizeAndType);
 }
 
@@ -1274,7 +1274,7 @@ const FTypeRegistry::FTypeInfo* FTypeRegistry::AddVersion4(const void* TraceData
 		int8 TypeSize = (int8)(1 << (Field.TypeInfo & Protocol0::Field_Pow2SizeMask));
 		if (Field.TypeInfo & Protocol0::Field_Float)
 		{
-			TypeSize = -TypeSize;
+			TypeSize = static_cast<int8>(-TypeSize);
 		}
 
 		auto& OutField = Builder.AddField(NameCursor, Field.NameSize, Field.Size);
@@ -1334,7 +1334,7 @@ const FTypeRegistry::FTypeInfo* FTypeRegistry::AddVersion6(const void* TraceData
 			int8 TypeSize = (int8)(1 << (Field.Regular.TypeInfo & Protocol0::Field_Pow2SizeMask));
 			if (Field.Regular.TypeInfo & Protocol0::Field_Float)
 			{
-				TypeSize = -TypeSize;
+				TypeSize = static_cast<int8>(-TypeSize);
 			}
 
 			auto& OutField = Builder.AddField(NameCursor, Field.Regular.NameSize, Field.Regular.Size);
diff --git a/Engine/Source/Developer/TraceServices/Private/Analyzers/AllocationsAnalysis.cpp b/Engine/Source/Developer/TraceServices/Private/Analyzers/AllocationsAnalysis.cpp
--- a/Engine/Source/Developer/TraceServices/Private/Analyzers/AllocationsAnalysis.cpp
+++ b/Engine/Source/Developer/TraceServices/Private/Analyzers/AllocationsAnalysis.cpp
@@ -160,7 +160,6 @@ bool FAllocationsAnalyzer::OnEvent(uint16 RouteId, EStyle Style, const FOnEventC
 
 			uint64 SizeUpper = EventData.GetValue<uint32>("Size");
 			const uint8 SizeLowerMask = static_cast<uint8>((1u << SizeShift) - 1u);
-			const uint8 AlignmentMask = ~SizeLowerMask;
 			const uint8 AlignmentPow2_SizeLower = EventData.GetValue<uint8>("AlignmentPow2_SizeLower");
 			const uint64 Size = (SizeUpper << SizeShift) | static_cast<uint64>(AlignmentPow2_SizeLower & SizeLowerMask);
 			const uint32 Alignment = 1u << (AlignmentPow2_SizeLower >> SizeShift);
diff --git a/Engine/Source/Runtime/Experimental/IoStore/OnDemand/Private/OnDemandBackendStatus.h b/Engine/Source/Runtime/Experimental/IoStore/OnDemand/Private/OnDemandBackendStatus.h
--- a/Engine/Source/Runtime/Experimental/IoStore/OnDemand/Private/OnDemandBackendStatus.h
+++ b/Engine/Source/Runtime/Experimental/IoStore/OnDemand/Private/OnDemandBackendStatus.h
@@ -88,7 +88,7 @@ private:
 
 	uint8 RemoveFlags(EFlags FlagsToRemove)
 	{
-		return Flags.fetch_and(~uint8(FlagsToRemove));
+		return Flags.fetch_and(static_cast<uint8>(~uint8(FlagsToRemove)));
 	}
 
 	uint8 AddOrRemoveFlags(EFlags FlagsToAddOrRemove, bool bValue)
