[Insights] Fixed UnrealInsights app not terminating when closed while "listening for direct trace stream". The socket needs to be shutdown before closing the asio acceptor.

#FailedValidations Preflight

diff --git a/Engine/Source/Developer/TraceAnalysis/Private/DataStream.cpp b/Engine/Source/Developer/TraceAnalysis/Private/DataStream.cpp
--- a/Engine/Source/Developer/TraceAnalysis/Private/DataStream.cpp
+++ b/Engine/Source/Developer/TraceAnalysis/Private/DataStream.cpp
@@ -12,13 +12,17 @@
 
 #include "Sockets.h"
 
+#if PLATFORM_LINUX
+#include <sys/socket.h>
+#endif
+
 namespace UE {
 namespace Trace {
 
 //--------------------------------------------------------------------
 // FFileDataStream
 //--------------------------------------------------------------------
-	
+
 FFileDataStream::FFileDataStream()
 	: Handle(nullptr)
 	, Remaining(0)
@@ -76,7 +80,7 @@ void FFileDataStream::Close()
 //--------------------------------------------------------------------
 // FTraceDataStream
 //--------------------------------------------------------------------
-	
+
 FTraceDataStream::FTraceDataStream(asio::ip::tcp::socket& InSocket)
 : Socket(MoveTemp(InSocket))
 {
@@ -139,7 +143,7 @@ int32 FTraceDataStream::Read(void* Dest, uint32 DestSize)
 		return int32(BytesRead);
 	}
 }
-	
+
 //--------------------------------------------------------------------
 // FDirectSocketContext
 //--------------------------------------------------------------------
@@ -154,7 +158,7 @@ struct FDirectSocketContext
 		: Context(1)
 		, Acceptor(Context)
 	{}
-	
+
 	asio::io_context Context;
 	asio::ip::tcp::acceptor Acceptor;
 };
@@ -162,7 +166,7 @@ struct FDirectSocketContext
 //--------------------------------------------------------------------
 // FDirectSocketStream
 //--------------------------------------------------------------------
-	
+
 FDirectSocketStream::FDirectSocketStream()
 	: ConnectionEvent(FPlatformProcess::GetSynchEventFromPool(false))
 {
@@ -179,13 +183,24 @@ FDirectSocketStream::FDirectSocketStream(TUniquePtr<FArchiveFileWriterGeneric>&&
 
 FDirectSocketStream::~FDirectSocketStream()
 {
-	FDirectSocketStream::Close();
-	// Make sure structures that rely on the asio context are
-	// destructed before it is removed.
-	ListeningThread.Reset();
+	if (AsioContext)
+	{
+#if PLATFORM_LINUX
+		// Socket needs to be shutdown before closing the asio acceptor, otherwise ListeningThread would deadlock.
+		shutdown(AsioContext->Acceptor.native_handle(), SHUT_RD);
+#endif
+		AsioContext->Acceptor.close();
+	}
+
+	if (ListeningThread)
+	{
+		ListeningThread->WaitForCompletion();
+		ListeningThread.Reset();
+	}
+
 	InternalStream.Reset();
+	AsioContext.Reset();
 }
-	
 
 uint16 FDirectSocketStream::StartListening(uint16 Port)
 {
@@ -213,7 +228,6 @@ uint16 FDirectSocketStream::StartListening(uint16 Port)
 	return Port;
 }
 
-
 bool FDirectSocketStream::CreateSocket(uint16 Port)
 {
 	using asio::ip::tcp;
@@ -237,7 +251,7 @@ bool FDirectSocketStream::CreateSocket(uint16 Port)
 	// Reuse address to avoid locking up port
 	asio::socket_base::reuse_address reuse_option(true);
 	Acceptor.set_option(reuse_option);
-	
+
 	// Setup endpoint for port on any address
 	tcp::endpoint Endpoint(tcp::v4(), uint16(Port));
 
@@ -266,7 +280,7 @@ void FDirectSocketStream::Accept()
 	{
 		return;
 	}
-	
+
 	tcp::socket Socket(AsioContext->Context);
 	asio::error_code ErrorCode;
 	AsioContext->Acceptor.accept(Socket, ErrorCode);
@@ -280,18 +294,18 @@ void FDirectSocketStream::Accept()
 		UE_LOG(LogCore, Error, TEXT("Failed accept sockets connection, error: %u."), ErrorCode.value());
 		return;
 	}
-	
+
 #if PLATFORM_WINDOWS
 	SetHandleInformation(HANDLE(SOCKET(Socket.native_handle())), HANDLE_FLAG_INHERIT, 0);
 #endif
-	
+
 	{
 		const auto Address = Socket.remote_endpoint().address();
 		TStringBuilder<16> RemoteAddress;
 		RemoteAddress << ANSI_TO_TCHAR(Address.to_string().c_str());
 		UE_LOG(LogCore, Log, TEXT("Accepted direct trace connection from %s"), RemoteAddress.ToString());
 	}
-	
+
 	AsioContext->Acceptor.close();
 	InternalStream = MakeUnique<FTraceDataStream>(Socket);
 	ConnectionEvent->Trigger();
@@ -324,7 +338,7 @@ bool FDirectSocketStream::BindToExistingConnection(FSocket&& ExistingConnectedSo
 
 uint32 FDirectSocketStream::Run()
 {
-	// Listening thread entry point. We currently only accept the first connection 
+	// Listening thread entry point. We currently only accept the first connection
 	// so the thread will exit immediately after.
 	Accept();
 	return 0;
@@ -371,12 +385,11 @@ void FDirectSocketStream::Close()
 	}
 
 	if (FileWriter.IsValid())
-    {
-    	FileWriter->Close();
-    	FileWriter.Reset();
-    }
+	{
+		FileWriter->Close();
+		FileWriter.Reset();
+	}
 }
 
-
 } // namespace Trace
 } // namespace UE
