Trace gap detection uses three sync points.

Previous logic used two sync points to calculate gaps, but the engine sends out three. In some scenarios this caused an error in gap detection which meant a higher serial was chosen as next serial. This could mean that large parts of the events were skipped. With this change gaps are not collected until the second sync packet has been parsed, and final gap determination is not done until the third sync packet.

diff --git a/Engine/Source/Developer/TraceAnalysis/Private/Analysis/Engine.cpp b/Engine/Source/Developer/TraceAnalysis/Private/Analysis/Engine.cpp
--- a/Engine/Source/Developer/TraceAnalysis/Private/Analysis/Engine.cpp
+++ b/Engine/Source/Developer/TraceAnalysis/Private/Analysis/Engine.cpp
@@ -4709,8 +4709,10 @@ void FProtocol5Stage::DetectSerialGaps(TArray<FEventDescStream>& EventDescHeap)
 
 	SyncCount = Transport.GetSyncCount();
 	UE_TRACE_ANALYSIS_DEBUG_LOG("SyncCount: %d (%d previous serial gaps)", SyncCount, SerialGaps.Num());
-
-	if (SyncCount == 1)
+	
+	// We expect 3 syncpoints before we consider the stream complete. Once the second
+	// sync is parsed we can collect the gaps in serials
+	if (SyncCount == 2)
 	{
 		// On the first update we will just collect gaps.
 		auto GatherGap = [this] (int32 Lhs, int32 Rhs)
@@ -4722,9 +4724,9 @@ void FProtocol5Stage::DetectSerialGaps(TArray<FEventDescStream>& EventDescHeap)
 		};
 		ForEachSerialGap(EventDescHeap, GatherGap);
 	}
-	else
+	else if (SyncCount == 3)
 	{
-		// On the second update we detect where gaps from the previous update
+		// On the third update we detect where gaps from the previous update
 		// start getting filled in. Any gaps preceding that point are genuine.
 		uint32 GapCount = 0;
 		auto RecordGap = [this, &GapCount] (int32 Lhs, int32 Rhs)
