[TraceAnalysis]
* Added OnVersion(uint32 TransportVersion, uint32 ProtocolVersion) callback to IAnalyzer.
* Added GetTypeSize() to FEventFieldInfo.
* Added GetFlags() to FEventTypeInfo, compatible with Protocol6::EEventFlags.
* Added IsImportant(), IsSync(), IsNoSync(), IsDefinition() and MaybeHasAux() getters to FEventTypeInfo.
* Simplified event flags used in FDispatch and FDispatchBuilder to be compatible with IAnalyzer::EEventFlags and Protocol6::EEventFlags.
* Updated TraceAnalyzer regarding usage of FEventTypeInfo::GetFlags().

#Synced-CL 49171203

diff --git a/Engine/Source/Developer/TraceAnalysis/Private/Analysis/Engine.cpp b/Engine/Source/Developer/TraceAnalysis/Private/Analysis/Engine.cpp
--- a/Engine/Source/Developer/TraceAnalysis/Private/Analysis/Engine.cpp
+++ b/Engine/Source/Developer/TraceAnalysis/Private/Analysis/Engine.cpp
@@ -369,14 +369,6 @@ double IAnalyzer::FEventTime::GetBaseDateTimestamp() const
 ////////////////////////////////////////////////////////////////////////////////
 struct FDispatch
 {
-	enum
-	{
-		Flag_Important		= 1 << 0,
-		Flag_MaybeHasAux	= 1 << 1,
-		Flag_NoSync			= 1 << 2,
-		Flag_Definition		= 1 << 3,
-	};
-
 	struct FField
 	{
 		uint32		Hash;
@@ -393,11 +385,17 @@ struct FDispatch
 	uint32			Hash				= 0;
 	uint16			Uid					= 0;
 	uint8			FieldCount			= 0;
-	uint8			Flags				= 0;
+	uint8			Flags				= 0;	// IAnalyzer::EEventFlags
 	uint16			EventSize			= 0;
 	uint16			LoggerNameOffset	= 0;	// From FDispatch ptr
 	uint16			EventNameOffset		= 0;	// From FDispatch ptr
 	FField			Fields[];
+
+	inline bool IsImportant() const  { return (Flags & (uint8)IAnalyzer::EEventFlags::Important) != 0; }
+	inline bool IsNoSync() const     { return (Flags & (uint8)IAnalyzer::EEventFlags::NoSync) != 0; }
+	inline bool IsSync() const       { return (Flags & (uint8)IAnalyzer::EEventFlags::NoSync) == 0; }
+	inline bool IsDefinition() const { return (Flags & (uint8)IAnalyzer::EEventFlags::Definition) != 0; }
+	inline bool MaybeHasAux() const  { return (Flags & (uint8)IAnalyzer::EEventFlags::MaybeHasAux) != 0; }
 };
 
 ////////////////////////////////////////////////////////////////////////////////
@@ -429,10 +427,7 @@ public:
 	void				SetUid(uint16 Uid);
 	void				SetLoggerName(const ANSICHAR* Name, int32 NameSize=-1);
 	void				SetEventName(const ANSICHAR* Name, int32 NameSize=-1);
-	void				SetImportant();
-	void				SetMaybeHasAux();
-	void				SetNoSync();
-	void				SetDefinition();
+	void				SetFlags(IAnalyzer::EEventFlags Flags);
 	FDispatch::FField&	AddField(const ANSICHAR* Name, int32 NameSize, uint16 Size);
 	FDispatch*			Finalize();
 
@@ -503,31 +498,10 @@ void FDispatchBuilder::SetEventName(const ANSICHAR* Name, int32 NameSize)
 }
 
 ////////////////////////////////////////////////////////////////////////////////
-void FDispatchBuilder::SetImportant()
+void FDispatchBuilder::SetFlags(IAnalyzer::EEventFlags Flags)
 {
 	auto* Dispatch = (FDispatch*)(Buffer.GetData());
-	Dispatch->Flags |= FDispatch::Flag_Important;
-}
-
-////////////////////////////////////////////////////////////////////////////////
-void FDispatchBuilder::SetMaybeHasAux()
-{
-	auto* Dispatch = (FDispatch*)(Buffer.GetData());
-	Dispatch->Flags |= FDispatch::Flag_MaybeHasAux;
-}
-
-////////////////////////////////////////////////////////////////////////////////
-void FDispatchBuilder::SetNoSync()
-{
-	auto* Dispatch = (FDispatch*)(Buffer.GetData());
-	Dispatch->Flags |= FDispatch::Flag_NoSync;
-}
-
-////////////////////////////////////////////////////////////////////////////////
-void FDispatchBuilder::SetDefinition()
-{
-	auto* Dispatch = (FDispatch*)(Buffer.GetData());
-	Dispatch->Flags |= FDispatch::Flag_Definition;
+	Dispatch->Flags = (uint8)Flags;
 }
 
 ////////////////////////////////////////////////////////////////////////////////
@@ -577,13 +551,11 @@ uint32 IAnalyzer::FEventTypeInfo::GetId() const
 }
 
 ////////////////////////////////////////////////////////////////////////////////
-#if UE_TRACE_ANALYSIS_DEBUG_API
-uint8 IAnalyzer::FEventTypeInfo::GetFlags() const
+IAnalyzer::EEventFlags IAnalyzer::FEventTypeInfo::GetFlags() const
 {
 	const auto* Inner = (const FDispatch*)this;
-	return Inner->Flags;
+	return static_cast<IAnalyzer::EEventFlags>(Inner->Flags);
 }
-#endif // UE_TRACE_ANALYSIS_DEBUG_API
 
 ////////////////////////////////////////////////////////////////////////////////
 const ANSICHAR* IAnalyzer::FEventTypeInfo::GetName() const
@@ -678,7 +650,7 @@ IAnalyzer::FEventFieldInfo::EType IAnalyzer::FEventFieldInfo::GetType() const
 
 	if (Inner->RefUid != 0)
 	{
-		switch(Inner->Size)
+		switch (Inner->Size)
 		{
 			case sizeof(uint8): return EType::Reference8;
 			case sizeof(uint16): return EType::Reference16;
@@ -706,6 +678,34 @@ IAnalyzer::FEventFieldInfo::EType IAnalyzer::FEventFieldInfo::GetType() const
 	return EType::None;
 }
 
+////////////////////////////////////////////////////////////////////////////////
+uint8 IAnalyzer::FEventFieldInfo::GetTypeSize() const
+{
+	const auto* Inner = (const FDispatch::FField*)this;
+
+	if (Inner->RefUid != 0)
+	{
+		return uint8(Inner->Size);
+	}
+
+	if (Inner->Class == UE::Trace::Protocol0::Field_String)
+	{
+		return 0;
+	}
+
+	if (Inner->SizeAndType > 0)
+	{
+		return uint8(Inner->SizeAndType);
+	}
+
+	if (Inner->SizeAndType < 0)
+	{
+		return uint8(-Inner->SizeAndType);
+	}
+
+	return 0;
+}
+
 ////////////////////////////////////////////////////////////////////////////////
 bool IAnalyzer::FEventFieldInfo::IsArray() const
 {
@@ -1048,11 +1048,7 @@ uint32 IAnalyzer::FEventData::GetAuxSize() const
 {
 	const auto* Info = (const FEventDataInfo*)this;
 
-	if ((Info->Dispatch.Flags & UE::Trace::FDispatch::Flag_MaybeHasAux) == 0)
-	{
-		return 0;
-	}
-	else
+	if (Info->Dispatch.MaybeHasAux())
 	{
 		uint32 Size = 0;
 		if (Info->AuxCollector)
@@ -1066,6 +1062,10 @@ uint32 IAnalyzer::FEventData::GetAuxSize() const
 		++Size; // aux terminator
 		return Size;
 	}
+	else
+	{
+		return 0;
+	}
 }
 #endif // UE_TRACE_ANALYSIS_DEBUG_API
 
@@ -1112,11 +1112,11 @@ uint32 IAnalyzer::FEventData::GetTotalSize(IAnalyzer::EStyle Style, const IAnaly
 	{
 		Size += 1; // uid
 	}
-	else if ((Info->Dispatch.Flags & UE::Trace::FDispatch::Flag_Important) != 0)
+	else if (Info->Dispatch.IsImportant())
 	{
 		Size += 4; // sizeof(UE::Trace::Protocol5::FImportantEventHeader);
 	}
-	else if (Info->Dispatch.Flags & UE::Trace::FDispatch::Flag_NoSync)
+	else if (Info->Dispatch.IsNoSync())
 	{
 		if (ProtocolVersion < 5)
 		{
@@ -1154,7 +1154,7 @@ bool IAnalyzer::FEventData::IsDefinitionImpl(uint32& OutTypeId) const
 {
 	const auto* Info = (const FEventDataInfo*)this;
 	OutTypeId = Info->Dispatch.Uid;
-	return (Info->Dispatch.Flags & FDispatch::Flag_Definition) != 0;
+	return Info->Dispatch.IsDefinition();
 }
 
 ////////////////////////////////////////////////////////////////////////////////
@@ -1287,20 +1287,12 @@ const FTypeRegistry::FTypeInfo* FTypeRegistry::AddVersion4(const void* TraceData
 		NameCursor += Field.NameSize;
 	}
 
-	if (NewEvent.Flags & uint32(Protocol4::EEventFlags::Important))
-	{
-		Builder.SetImportant();
-	}
-
-	if (NewEvent.Flags & uint32(Protocol4::EEventFlags::MaybeHasAux))
-	{
-		Builder.SetMaybeHasAux();
-	}
-
-	if (NewEvent.Flags & uint32(Protocol4::EEventFlags::NoSync))
-	{
-		Builder.SetNoSync();
-	}
+	// NewEvent.Flags contains Protocol1::EEventFlags,
+	// so it can be used in place for IAnalyzer::EEventFlags.
+	static_assert((uint8)UE::Trace::Protocol1::EEventFlags::Important == (uint8)IAnalyzer::EEventFlags::Important, "");
+	static_assert((uint8)UE::Trace::Protocol1::EEventFlags::MaybeHasAux == (uint8)IAnalyzer::EEventFlags::MaybeHasAux, "");
+	static_assert((uint8)UE::Trace::Protocol1::EEventFlags::NoSync == (uint8)IAnalyzer::EEventFlags::NoSync, "");
+	Builder.SetFlags(static_cast<IAnalyzer::EEventFlags>(NewEvent.Flags));
 
 	FTypeInfo* TypeInfo = Builder.Finalize();
 	Add(TypeInfo);
@@ -1380,25 +1372,14 @@ const FTypeRegistry::FTypeInfo* FTypeRegistry::AddVersion6(const void* TraceData
 		}
 	}
 
-	if (NewEvent.Flags & uint32(EEventFlags::Important))
-	{
-		Builder.SetImportant();
-	}
-
-	if (NewEvent.Flags & uint32(EEventFlags::MaybeHasAux))
-	{
-		Builder.SetMaybeHasAux();
-	}
-
-	if (NewEvent.Flags & uint32(EEventFlags::NoSync))
-	{
-		Builder.SetNoSync();
-	}
-
-	if (NewEvent.Flags & uint32(EEventFlags::Definition))
-	{
-		Builder.SetDefinition();
-	}
+	// NewEvent.Flags contains Protocol6::EEventFlags,
+	// so it can be used in place for IAnalyzer::EEventFlags.
+	static_assert((uint8)UE::Trace::Protocol6::EEventFlags::None == (uint8)IAnalyzer::EEventFlags::None, "");
+	static_assert((uint8)UE::Trace::Protocol6::EEventFlags::Important == (uint8)IAnalyzer::EEventFlags::Important, "");
+	static_assert((uint8)UE::Trace::Protocol6::EEventFlags::MaybeHasAux == (uint8)IAnalyzer::EEventFlags::MaybeHasAux, "");
+	static_assert((uint8)UE::Trace::Protocol6::EEventFlags::NoSync == (uint8)IAnalyzer::EEventFlags::NoSync, "");
+	static_assert((uint8)UE::Trace::Protocol6::EEventFlags::Definition == (uint8)IAnalyzer::EEventFlags::Definition, "");
+	Builder.SetFlags(static_cast<IAnalyzer::EEventFlags>(NewEvent.Flags));
 
 	FTypeInfo* TypeInfo = Builder.Finalize();
 	Add(TypeInfo);
@@ -1451,9 +1432,7 @@ public:
 	void				OnNewType(const FTypeRegistry::FTypeInfo* TypeInfo);
 	void				OnEvent(const FTypeRegistry::FTypeInfo& TypeInfo, IAnalyzer::EStyle Style, const IAnalyzer::FOnEventContext& Context);
 	void				OnThreadInfo(const FThreads::FInfo& ThreadInfo);
-#if UE_TRACE_ANALYSIS_DEBUG_API
 	void				OnVersion(uint32 TransportVersion, uint32 ProtocolVersion);
-#endif // UE_TRACE_ANALYSIS_DEBUG_API
 
 private:
 	void				BuildRoutes();
@@ -1593,13 +1572,12 @@ void FAnalyzerHub::OnNewType(const FTypeRegistry::FTypeInfo* TypeInfo)
 #if UE_TRACE_ANALYSIS_DEBUG
 	const IAnalyzer::FEventTypeInfo& EventTypeInfo = *(IAnalyzer::FEventTypeInfo*)TypeInfo;
 #if UE_TRACE_ANALYSIS_DEBUG_API
-	const uint8 EventTypeFlags = EventTypeInfo.GetFlags();
 	UE_TRACE_ANALYSIS_DEBUG_LOG("NewEvent: Uid=%u (%s.%s) Flags=%s%s%s%s",
 		EventTypeInfo.GetId(), EventTypeInfo.GetLoggerName(), EventTypeInfo.GetName(),
-		(EventTypeFlags & FDispatch::Flag_Definition) ? "Definition|" : "",
-		(EventTypeFlags & FDispatch::Flag_Important) ? "Important|" : "",
-		(EventTypeFlags & FDispatch::Flag_MaybeHasAux) ? "MaybeHasAux|" : "",
-		(EventTypeFlags & FDispatch::Flag_NoSync) ? "NoSync" : "Sync");
+		EventTypeInfo.IsDefinition() ? "Definition|" : "",
+		EventTypeInfo.IsImportant() ? "Important|" : "",
+		EventTypeInfo.MaybeHasAux() ? "MaybeHasAux|" : "",
+		EventTypeInfo.IsNoSync() ? "NoSync" : "Sync");
 #else // UE_TRACE_ANALYSIS_DEBUG_API
 	UE_TRACE_ANALYSIS_DEBUG_LOG("NewEvent: Uid=%u (%s.%s)",
 		EventTypeInfo.GetId(), EventTypeInfo.GetLoggerName(), EventTypeInfo.GetName());
@@ -1655,7 +1633,6 @@ void FAnalyzerHub::OnThreadInfo(const FThreads::FInfo& ThreadInfo)
 }
 
 ////////////////////////////////////////////////////////////////////////////////
-#if UE_TRACE_ANALYSIS_DEBUG_API
 void FAnalyzerHub::OnVersion(uint32 TransportVersion, uint32 ProtocolVersion)
 {
 	for (IAnalyzer* Analyzer : Analyzers)
@@ -1666,7 +1643,6 @@ void FAnalyzerHub::OnVersion(uint32 TransportVersion, uint32 ProtocolVersion)
 		}
 	}
 }
-#endif // UE_TRACE_ANALYSIS_DEBUG_API
 
 ////////////////////////////////////////////////////////////////////////////////
 void FAnalyzerHub::BuildRoutes()
@@ -2085,9 +2061,7 @@ public:
 	void				LeaveScope(uint64 RelativeTimestamp);
 	void				LeaveScopeA(uint64 AbsoluteTimestamp);
 	void				LeaveScopeB(uint64 BaseRelativeTimestamp);
-#if UE_TRACE_ANALYSIS_DEBUG_API
 	void				OnVersion(uint32 TransportVersion, uint32 ProtocolVersion) { AnalyzerHub.OnVersion(TransportVersion, ProtocolVersion); }
-#endif // UE_TRACE_ANALYSIS_DEBUG_API
 #if UE_TRACE_ANALYSIS_DEBUG
 	void				DebugLogEvent(const FTypeRegistry::FTypeInfo* TypeInfo, uint32 FixedSize, uint32 AuxSize, uint32 Serial);
 	void				DebugLogNewEvent(uint32 Uid, const FTypeRegistry::FTypeInfo* TypeInfo, uint32 EventSize);
@@ -2325,7 +2299,7 @@ void FAnalysisBridge::DebugLogEvent(const FTypeRegistry::FTypeInfo* TypeInfo, ui
 	}
 #endif // UE_TRACE_ANALYSIS_DEBUG_LEVEL
 
-	if ((TypeInfo->Flags & FTypeRegistry::FTypeInfo::Flag_NoSync) == 0)
+	if (TypeInfo->IsSync())
 	{
 		++State.SyncEventCount;
 		State.SyncEventSize += EventSize;
@@ -2333,7 +2307,7 @@ void FAnalysisBridge::DebugLogEvent(const FTypeRegistry::FTypeInfo* TypeInfo, ui
 		UE_TRACE_ANALYSIS_DEBUG_Appendf(" Sync=%u", Serial);
 #endif // UE_TRACE_ANALYSIS_DEBUG_LEVEL
 	}
-	else if ((TypeInfo->Flags & FTypeRegistry::FTypeInfo::Flag_Important) != 0)
+	else if (TypeInfo->IsImportant())
 	{
 		++State.ImportantNoSyncEventCount;
 		State.ImportantNoSyncEventSize += EventSize;
@@ -2783,7 +2757,7 @@ FProtocol2Stage::FProtocol2Stage(uint32 Version, FTransport* InTransport)
 			Dispatch.SetUid(uint16(Protocol3::EKnownEventUids::NewEvent));
 			Dispatch.SetLoggerName("$Trace");
 			Dispatch.SetEventName("NewEvent");
-			Dispatch.SetNoSync();
+			Dispatch.SetFlags(IAnalyzer::EEventFlags::NoSync);
 			TypeRegistry.Add(Dispatch.Finalize());
 		}
 		break;
@@ -3039,7 +3013,7 @@ int32 FProtocol2Stage::OnData(FStreamReader& Reader, FAnalysisBridge& Bridge)
 		uint32 BlockSize = Header->Size;
 
 		// Make sure we consume events in the correct order
-		if ((TypeInfo->Flags & FTypeRegistry::FTypeInfo::Flag_NoSync) == 0)
+		if (TypeInfo->IsSync())
 		{
 			switch (ProtocolVersion)
 			{
@@ -3090,7 +3064,7 @@ int32 FProtocol2Stage::OnData(FStreamReader& Reader, FAnalysisBridge& Bridge)
 		uint32 AuxSize = 0;
 #endif
 		FAuxDataCollector AuxCollector;
-		if (TypeInfo->Flags & FTypeRegistry::FTypeInfo::Flag_MaybeHasAux)
+		if (TypeInfo->MaybeHasAux())
 		{
 #if UE_TRACE_ANALYSIS_DEBUG
 			auto AuxMark = Reader.SaveMark();
@@ -3108,7 +3082,7 @@ int32 FProtocol2Stage::OnData(FStreamReader& Reader, FAnalysisBridge& Bridge)
 #endif
 		}
 
-		if ((TypeInfo->Flags & FTypeRegistry::FTypeInfo::Flag_NoSync) == 0)
+		if (TypeInfo->IsSync())
 		{
 			Serial.Value += 1;
 			Serial.Value &= 0x3fffffff; // Don't set MSBs. They have other uses
@@ -3293,7 +3267,7 @@ int32 FProtocol4Stage::OnDataImpl(FStreamReader& Reader, FAnalysisBridge& Bridge
 	uint32 BlockSize = Header->Size;
 
 	// Make sure we consume events in the correct order
-	if ((TypeInfo->Flags & FTypeRegistry::FTypeInfo::Flag_NoSync) == 0)
+	if (TypeInfo->IsSync())
 	{
 		if (Reader.GetPointer<Protocol4::FEventHeaderSync>() == nullptr)
 		{
@@ -3330,11 +3304,11 @@ int32 FProtocol4Stage::OnDataImpl(FStreamReader& Reader, FAnalysisBridge& Bridge
 
 	// Collect auxiliary data
 	FAuxDataCollector AuxCollector;
-	if (TypeInfo->Flags & FTypeRegistry::FTypeInfo::Flag_MaybeHasAux)
+	if (TypeInfo->MaybeHasAux())
 	{
 		// Important events' size may include their array data so we need to backtrack
 		auto NextMark = Reader.SaveMark();
-		if (TypeInfo->Flags & FTypeRegistry::FTypeInfo::Flag_Important)
+		if (TypeInfo->IsImportant())
 		{
 			Reader.RestoreMark(Mark);
 			Reader.Advance(sizeof(Protocol4::FEventHeader) + TypeInfo->EventSize);
@@ -3358,7 +3332,7 @@ int32 FProtocol4Stage::OnDataImpl(FStreamReader& Reader, FAnalysisBridge& Bridge
 		// User error could have resulted in less space being used that was
 		// allocated for important events. So we can't assume that aux data
 		// reading has read all the way up to the next event. So we use marks
-		if (TypeInfo->Flags & FTypeRegistry::FTypeInfo::Flag_Important)
+		if (TypeInfo->IsImportant())
 		{
 			Reader.RestoreMark(NextMark);
 		}
@@ -3370,7 +3344,7 @@ int32 FProtocol4Stage::OnDataImpl(FStreamReader& Reader, FAnalysisBridge& Bridge
 #endif
 
 	// Maintain sync
-	if ((TypeInfo->Flags & FTypeRegistry::FTypeInfo::Flag_NoSync) == 0)
+	if (TypeInfo->IsSync())
 	{
 		Serial.Value += 1;
 		Serial.Value &= 0x7fffffff; // don't set msb. that has other uses
@@ -3805,14 +3779,14 @@ int32 FProtocol5Stage::ParseImportantEvents(FStreamReader& Reader, EventDescArra
 		EventDesc.EventSize = sizeof(*Header) + TypeInfo->EventSize;
 #endif
 
-		if (TypeInfo->Flags & FTypeRegistry::FTypeInfo::Flag_MaybeHasAux)
+		if (TypeInfo->MaybeHasAux())
 		{
 			EventDesc.bHasAux = 1;
 		}
 
 		OutEventDescs.Add(EventDesc);
 
-		if (TypeInfo->Flags & FTypeRegistry::FTypeInfo::Flag_MaybeHasAux)
+		if (TypeInfo->MaybeHasAux())
 		{
 			const uint8* Cursor = Header->Data + TypeInfo->EventSize;
 			const uint8* End = Header->Data + Header->Size;
@@ -3823,6 +3797,7 @@ int32 FProtocol5Stage::ParseImportantEvents(FStreamReader& Reader, EventDescArra
 					break;
 				}
 
+				check(Cursor[0] == uint8(EKnownUids::AuxData));
 				const auto* AuxHeader = (FAuxHeader*)Cursor;
 
 				FEventDesc& AuxDesc = OutEventDescs.Emplace_GetRef();
@@ -4604,9 +4579,9 @@ int32 FProtocol5Stage::ParseEvent(FStreamReader& Reader, FEventDesc& EventDesc,
 		}
 
 		EventSize = TypeInfo->EventSize;
-		EventDesc.bHasAux = !!(TypeInfo->Flags & FTypeRegistry::FTypeInfo::Flag_MaybeHasAux);
+		EventDesc.bHasAux = TypeInfo->MaybeHasAux();
 
-		if ((TypeInfo->Flags & FDispatch::Flag_NoSync) == 0)
+		if (TypeInfo->IsSync())
 		{
 			memcpy(&Serial, Cursor, sizeof(int32));
 			Serial &= ESerial::Mask;
@@ -4709,7 +4684,7 @@ void FProtocol5Stage::DetectSerialGaps(TArray<FEventDescStream>& EventDescHeap)
 
 	SyncCount = Transport.GetSyncCount();
 	UE_TRACE_ANALYSIS_DEBUG_LOG("SyncCount: %d (%d previous serial gaps)", SyncCount, SerialGaps.Num());
-	
+
 	// We expect 3 syncpoints before we consider the stream complete. Once the second
 	// sync is parsed we can collect the gaps in serials
 	if (SyncCount == 2)
@@ -4869,6 +4844,10 @@ int32 FProtocol5Stage::DispatchEvents(const FMachineContext& Context, const FEve
 	{
 		uint32 Uid = uint32(Cursor->Uid);
 
+#if UE_TRACE_ANALYSIS_DEBUG && UE_TRACE_ANALYSIS_DEBUG_LEVEL >= 4
+		UE_TRACE_ANALYSIS_DEBUG_LOG("Dispatch event (Uid=%u)...", Uid);
+#endif
+
 		if (DispatchKnownEvent(Context, Uid, Cursor))
 		{
 			continue;
@@ -5232,9 +5211,7 @@ FEstablishTransportStage::EStatus FEstablishTransportStage::OnData(
 
 	Reader.Advance(sizeof(*Header));
 
-#if UE_TRACE_ANALYSIS_DEBUG_API
 	Context.Bridge.OnVersion(TransportVersion, ProtocolVersion);
-#endif // UE_TRACE_ANALYSIS_DEBUG_API
 
 	return EStatus::Continue;
 }
diff --git a/Engine/Source/Developer/TraceAnalysis/Public/Trace/Analyzer.h b/Engine/Source/Developer/TraceAnalysis/Public/Trace/Analyzer.h
--- a/Engine/Source/Developer/TraceAnalysis/Public/Trace/Analyzer.h
+++ b/Engine/Source/Developer/TraceAnalysis/Public/Trace/Analyzer.h
@@ -6,6 +6,7 @@
 #include "Containers/ArrayView.h"
 #include "Containers/StringView.h"
 #include "TraceAnalysisDebug.h"
+#include "Trace/Detail/Protocol.h"
 #include "Trace/Trace.h"
 
 #include <type_traits>
@@ -74,8 +75,11 @@ public:
 		UE_API uint32 GetOffset() const;
 #endif // UE_TRACE_ANALYSIS_DEBUG_API
 
-		/** Gets the size in bytes for this field */
+		/** Gets the size in bytes for this field. */
 		UE_API uint8 GetSize() const;
+
+		/** Gets the size in bytes for the type of this field. */
+		UE_API uint8 GetTypeSize() const;
 	};
 
 	struct FEventFieldHandle
@@ -84,16 +88,28 @@ public:
 		int32	Detail;
 	};
 
+	enum class EEventFlags : uint8
+	{
+		None        = 0,
+		Important   = 1 << 0,
+		MaybeHasAux = 1 << 1,
+		NoSync      = 1 << 2,
+		Definition  = 1 << 3,
+	};
+
 	struct FEventTypeInfo
 	{
 		/** Each event is assigned a unique ID when logged. Note that this is not
 		 * guaranteed to be the same for the same event from one trace to the next. */
 		UE_API uint32 GetId() const;
 
-#if UE_TRACE_ANALYSIS_DEBUG_API
 		/** Returns the event's flags. */
-		UE_API uint8 GetFlags() const;
-#endif // UE_TRACE_ANALYSIS_DEBUG_API
+		UE_API EEventFlags GetFlags() const;
+		inline bool IsImportant() const  { return ((uint8)GetFlags() & (uint8)EEventFlags::Important) != 0; }
+		inline bool IsSync() const       { return ((uint8)GetFlags() & (uint8)EEventFlags::NoSync) == 0; }
+		inline bool IsNoSync() const     { return ((uint8)GetFlags() & (uint8)EEventFlags::NoSync) != 0; }
+		inline bool IsDefinition() const { return ((uint8)GetFlags() & (uint8)EEventFlags::Definition) != 0; }
+		inline bool MaybeHasAux() const  { return ((uint8)GetFlags() & (uint8)EEventFlags::MaybeHasAux) != 0; }
 
 		/** The name of the event. */
 		UE_API const ANSICHAR* GetName() const;
@@ -356,11 +372,13 @@ public:
 		return true;
 	}
 
-#if UE_TRACE_ANALYSIS_DEBUG_API
+	/** Called when information about trace transport and trace protocol versions are available.
+	 * This is called after OnAnalysisBegin(), but before any OnNewEvent() or OnEvent() calls.
+	 * @param TransportVersion The version of the trace transport (see Trace\Detail\Transport.h)
+	 * @param ProtocolVersion The version of the trace protocol (see Trace\Detail\Protocols\*.h) */
 	virtual void OnVersion(uint32 TransportVersion, uint32 ProtocolVersion)
 	{
 	}
-#endif // UE_TRACE_ANALYSIS_DEBUG_API
 
 private:
 	template <typename ValueType> static ValueType CoerceValue(const void* Addr, int8 SizeAndType);
diff --git a/Engine/Source/Programs/TraceAnalyzer/Private/ConvertToText.cpp b/Engine/Source/Programs/TraceAnalyzer/Private/ConvertToText.cpp
--- a/Engine/Source/Programs/TraceAnalyzer/Private/ConvertToText.cpp
+++ b/Engine/Source/Programs/TraceAnalyzer/Private/ConvertToText.cpp
@@ -991,27 +991,36 @@ bool FConvertToTextAnalyzer::OnNewEvent(uint16 RouteId, const FEventTypeInfo& Ty
 		Serializer.BeginAttribute();
 		Serializer.WriteKey("Flags");
 		Serializer.Append("\"");
-		const uint8 EventFlags = TypeInfo.GetFlags();
-		if ((EventFlags & (uint8)EEventFlags::Definition) != 0)
+
+		IAnalyzer::EEventFlags EventFlags = TypeInfo.GetFlags();
+		if (TypeInfo.IsDefinition())
 		{
 			Serializer.Append("Definition|");
+			EnumRemoveFlags(EventFlags, IAnalyzer::EEventFlags::Definition);
 		}
-		if ((EventFlags & (uint8)EEventFlags::Important) != 0)
+		if (TypeInfo.IsImportant())
 		{
 			Serializer.Append("Important|");
+			EnumRemoveFlags(EventFlags, IAnalyzer::EEventFlags::Important);
 		}
-		if ((EventFlags & (uint8)EEventFlags::MaybeHasAux) != 0)
+		if (TypeInfo.MaybeHasAux())
 		{
 			Serializer.Append("MaybeHasAux|");
+			EnumRemoveFlags(EventFlags, IAnalyzer::EEventFlags::MaybeHasAux);
 		}
-		if ((EventFlags & (uint8)EEventFlags::NoSync) != 0)
+		if (TypeInfo.IsNoSync())
 		{
 			Serializer.Append("NoSync");
+			EnumRemoveFlags(EventFlags, IAnalyzer::EEventFlags::NoSync);
 		}
 		else
 		{
 			Serializer.Append("Sync");
 		}
+		if ((uint32)EventFlags != 0)
+		{
+			Serializer.Appendf("|%u", (uint32)EventFlags);
+		}
 		Serializer.Append("\"");
 		Serializer.EndAttribute();
 
@@ -1177,11 +1186,11 @@ bool FConvertToTextAnalyzer::OnEvent(uint16 RouteId, EStyle Style, const FOnEven
 
 	EventsTotalSize += TotalEventSize;
 
-	const bool bHasAttachment = ((TypeInfo.GetFlags() & (uint8)EEventFlags::Important) == 0) &&
+	const bool bHasAttachment = !TypeInfo.IsImportant() &&
 								(AuxSize == 0) &&
 								(RawSize > FixedSize) &&
 								(AttachmentSize > 0) &&
-								EventData.GetAttachment() != nullptr;
+								(EventData.GetAttachment() != nullptr);
 
 	bool bEmitSizeStats = false;
 
