TraceAnalysis:
* Added initial FTraceWriter implementation.
* Added IOutDataStream and FFileOutDataStream.

#Synced-CL 49201375

diff --git a/Engine/Source/Developer/TraceAnalysis/Private/OutDataStream.cpp b/Engine/Source/Developer/TraceAnalysis/Private/OutDataStream.cpp
new file mode 100644
--- /dev/null
+++ b/Engine/Source/Developer/TraceAnalysis/Private/OutDataStream.cpp
@@ -0,0 +1,53 @@
+// Copyright Epic Games, Inc. All Rights Reserved.
+
+#include "Trace/OutDataStream.h"
+
+#include "HAL/PlatformFileManager.h"
+
+namespace UE::Trace {
+
+//--------------------------------------------------------------------
+// FFileOutDataStream
+//--------------------------------------------------------------------
+
+FFileOutDataStream::FFileOutDataStream()
+	: Handle(nullptr)
+{
+}
+
+FFileOutDataStream::~FFileOutDataStream()
+{
+}
+
+bool FFileOutDataStream::Open(const TCHAR* Path)
+{
+	Handle.Reset(FPlatformFileManager::Get().GetPlatformFile().OpenWrite(Path));
+	if (Handle == nullptr)
+	{
+		return false;
+	}
+	return true;
+}
+
+int32 FFileOutDataStream::Write(void* Data, uint32 Size)
+{
+	if (Handle == nullptr)
+	{
+		return -1;
+	}
+
+	if (!Handle->Write((uint8*)Data, Size))
+	{
+		Close();
+		return -1;
+	}
+
+	return Size;
+}
+
+void FFileOutDataStream::Close()
+{
+	Handle.Reset();
+}
+
+} // namespace UE::Trace
diff --git a/Engine/Source/Developer/TraceAnalysis/Private/TraceWriter.cpp b/Engine/Source/Developer/TraceAnalysis/Private/TraceWriter.cpp
new file mode 100644
--- /dev/null
+++ b/Engine/Source/Developer/TraceAnalysis/Private/TraceWriter.cpp
@@ -0,0 +1,1557 @@
+// Copyright Epic Games, Inc. All Rights Reserved.
+
+#include "Trace/TraceWriter.h"
+
+#include "HAL/PlatformTime.h"
+
+// TraceLog
+#include "Trace/Detail/Transport.h"
+
+// TraceAnalysis
+#include "Trace/OutDataStream.h"
+
+////////////////////////////////////////////////////////////////////////////////////////////////////
+
+namespace UE::Trace::Private
+{
+
+TRACELOG_API uint32 GetEncodeMaxSize(uint32 InputSize);
+TRACELOG_API int32 Encode(const void* Src, int32 SrcSize, void* Dest, int32 DestSize);
+TRACELOG_API int32 Decode(const void* Src, int32 SrcSize, void* Dest, int32 DestSize);
+
+} // namespace UE::Trace::Private
+
+////////////////////////////////////////////////////////////////////////////////////////////////////
+
+namespace UE::Trace::Private::WriterProtocol
+{
+	using Protocol7::EProtocol;
+	using Protocol7::EFieldType;
+	using Protocol7::FEventHeader;
+	using Protocol7::FImportantEventHeader;
+	using Protocol7::FEventHeaderSync;
+	using Protocol7::FAuxHeader;
+	using Protocol7::EEventFlags;
+	using Protocol7::EFieldFamily;
+	using Protocol7::FNewEventField;
+	using Protocol7::FNewEventEvent;
+	using Protocol7::EKnownEventUids;
+}
+
+////////////////////////////////////////////////////////////////////////////////////////////////////
+
+namespace UE::Trace
+{
+
+////////////////////////////////////////////////////////////////////////////////////////////////////
+// FTraceWriterFieldTypeResolver
+////////////////////////////////////////////////////////////////////////////////////////////////////
+
+static_assert(FTraceWriterFieldTypeResolver::GetType<uint64>() == ETraceWriterFieldType::Uint64, "");
+static_assert(FTraceWriterFieldTypeResolver::GetType<const uint64>() == ETraceWriterFieldType::Uint64, "");
+static_assert(FTraceWriterFieldTypeResolver::GetType<uint64[]>() == (ETraceWriterFieldType::ArrayFlag | ETraceWriterFieldType::Uint64), "");
+static_assert(FTraceWriterFieldTypeResolver::GetType<const uint64[]>() == (ETraceWriterFieldType::ArrayFlag | ETraceWriterFieldType::Uint64), "");
+static_assert(FTraceWriterFieldTypeResolver::GetType<TConstArrayView<uint64>>() == (ETraceWriterFieldType::ArrayFlag | ETraceWriterFieldType::Uint64), "");
+static_assert(FTraceWriterFieldTypeResolver::GetType<TArrayView<uint64>>() == (ETraceWriterFieldType::ArrayFlag | ETraceWriterFieldType::Uint64), "");
+
+static_assert(FTraceWriterFieldTypeResolver::GetType<void*>() == ETraceWriterFieldType::Pointer, "");
+static_assert(FTraceWriterFieldTypeResolver::GetType<const void*>() == ETraceWriterFieldType::Pointer, "");
+static_assert(FTraceWriterFieldTypeResolver::GetType<void* []>() == (ETraceWriterFieldType::ArrayFlag | ETraceWriterFieldType::Pointer), "");
+static_assert(FTraceWriterFieldTypeResolver::GetType<const void* []>() == (ETraceWriterFieldType::ArrayFlag | ETraceWriterFieldType::Pointer), "");
+static_assert(FTraceWriterFieldTypeResolver::GetType<TConstArrayView<void*>>() == (ETraceWriterFieldType::ArrayFlag | ETraceWriterFieldType::Pointer), "");
+static_assert(FTraceWriterFieldTypeResolver::GetType<TArrayView<void*>>() == (ETraceWriterFieldType::ArrayFlag | ETraceWriterFieldType::Pointer), "");
+static_assert(FTraceWriterFieldTypeResolver::GetType<TConstArrayView<const void*>>() == (ETraceWriterFieldType::ArrayFlag | ETraceWriterFieldType::Pointer), "");
+static_assert(FTraceWriterFieldTypeResolver::GetType<TArrayView<const void*>>() == (ETraceWriterFieldType::ArrayFlag | ETraceWriterFieldType::Pointer), "");
+
+////////////////////////////////////////////////////////////////////////////////////////////////////
+// FTraceWriterEventInfo
+////////////////////////////////////////////////////////////////////////////////////////////////////
+
+uint32 FTraceWriterEventInfo::GetFieldIndex(FAnsiStringView InFieldName) const
+{
+	const uint32 FieldCount = Fields.Num();
+	for (uint32 FieldIndex = 0; FieldIndex < FieldCount; ++FieldIndex)
+	{
+		if (InFieldName.Equals(Fields[FieldIndex].GetName(), ESearchCase::CaseSensitive))
+		{
+			return FieldIndex;
+		}
+	}
+	return FieldCount;
+}
+
+////////////////////////////////////////////////////////////////////////////////////////////////////
+// FTraceWriterEventDeclarationBuilder
+////////////////////////////////////////////////////////////////////////////////////////////////////
+
+FTraceWriterEventDeclarationBuilder::~FTraceWriterEventDeclarationBuilder()
+{
+	if (!bIsCompleted)
+	{
+		Writer.SetError(TEXT("An event declaration builder is destroyed without being completed. Check if an End() call is missing."));
+	}
+}
+
+////////////////////////////////////////////////////////////////////////////////////////////////////
+
+FTraceWriterEventDeclarationBuilder& FTraceWriterEventDeclarationBuilder::Field(FAnsiStringView InFieldName, ETraceWriterFieldType InFieldType)
+{
+	if (!EnumHasAnyFlags(InFieldType, ETraceWriterFieldType::ReferenceFlag | ETraceWriterFieldType::DefinitionIdFlag))
+	{
+		FTraceWriterEventFieldInfo& Field = EventInfo.Fields.Emplace_GetRef(InFieldName, InFieldType);
+		Field.Size = FTraceWriter::GetByteSizeForFieldType(InFieldType);
+		if (Field.Size == 0)
+		{
+			EventInfo.Flags |= ETraceWriterEventFlags::MaybeHasAux;
+		}
+	}
+	else
+	{
+		Writer.SetError(TEXT("Invalid type for a regular event field!"));
+	}
+	return *this;
+}
+
+////////////////////////////////////////////////////////////////////////////////////////////////////
+
+FTraceWriterEventDeclarationBuilder& FTraceWriterEventDeclarationBuilder::ReferenceField(FAnsiStringView InFieldName, ETraceWriterFieldType InFieldType, uint16 InRefUid)
+{
+	if (InFieldType == ETraceWriterFieldType::Reference8 ||
+		InFieldType == ETraceWriterFieldType::Reference16 ||
+		InFieldType == ETraceWriterFieldType::Reference32 ||
+		InFieldType == ETraceWriterFieldType::Reference64)
+	{
+		FTraceWriterEventFieldInfo& Field = EventInfo.Fields.Emplace_GetRef(InFieldName, InFieldType);
+		Field.Size = FTraceWriter::GetByteSizeForFieldType(InFieldType);
+		check(Field.Size > 0);
+		Field.RefUid = InRefUid;
+	}
+	else
+	{
+		Writer.SetError(TEXT("Invalid type for a reference event field!"));
+	}
+	return *this;
+}
+
+////////////////////////////////////////////////////////////////////////////////////////////////////
+
+FTraceWriterEventDeclarationBuilder& FTraceWriterEventDeclarationBuilder::DefinitionIdField(ETraceWriterFieldType InFieldType)
+{
+	if (InFieldType == ETraceWriterFieldType::DefinitionId8 ||
+		InFieldType == ETraceWriterFieldType::DefinitionId16 ||
+		InFieldType == ETraceWriterFieldType::DefinitionId32 ||
+		InFieldType == ETraceWriterFieldType::DefinitionId64)
+	{
+		FTraceWriterEventFieldInfo& Field = EventInfo.Fields.Emplace_GetRef(ANSITEXTVIEW("DefinitionId"), InFieldType);
+		Field.Size = FTraceWriter::GetByteSizeForFieldType(InFieldType);
+		check(Field.Size > 0);
+	}
+	else
+	{
+		Writer.SetError(TEXT("Invalid type for a definition event field!"));
+	}
+	return *this;
+}
+
+////////////////////////////////////////////////////////////////////////////////////////////////////
+
+uint32 FTraceWriterEventDeclarationBuilder::End()
+{
+	// Write the event declaration to the stream...
+	Writer.WriteEventDeclaration(*this);
+
+	const uint32 EventId = EventInfo.Id;
+
+	// End() destroys the current EventDeclarationBuilder!
+	bIsCompleted = true;
+	delete this;
+
+	return EventId;
+}
+
+////////////////////////////////////////////////////////////////////////////////////////////////////
+// FTraceWriterFieldData
+////////////////////////////////////////////////////////////////////////////////////////////////////
+
+FTraceWriterFieldData::~FTraceWriterFieldData()
+{
+	if (Size > FTraceWriterFieldData::MaxInlineDataSize)
+	{
+		check(Buffer != nullptr);
+		FMemory::Free(Buffer);
+	}
+	Data = 0;
+	Data2 = 0;
+	Size = 0;
+}
+
+////////////////////////////////////////////////////////////////////////////////////////////////////
+
+void FTraceWriterFieldData::SetData(const void* InData, uint32 InDataSize)
+{
+	if (InDataSize <= FTraceWriterFieldData::MaxInlineDataSize)
+	{
+		FMemory::Memcpy(&Data, InData, InDataSize);
+	}
+	else
+	{
+		check(Buffer == nullptr);
+		Buffer = FMemory::Malloc(InDataSize, 0);
+		FMemory::Memcpy(Buffer, InData, InDataSize);
+	}
+	Size = InDataSize;
+}
+
+////////////////////////////////////////////////////////////////////////////////////////////////////
+// FTraceWriterEventBuilder
+////////////////////////////////////////////////////////////////////////////////////////////////////
+
+FTraceWriterEventBuilder::~FTraceWriterEventBuilder()
+{
+	if (!bIsCompleted)
+	{
+		Writer.SetError(TEXT("An event builder is destroyed without being completed. Check if an End() call is missing."));
+	}
+}
+
+////////////////////////////////////////////////////////////////////////////////////////////////////
+
+FTraceWriterEventBuilder& FTraceWriterEventBuilder::Field(uint32 Index, const void* Data, uint32 DataSize)
+{
+	if (Writer.IsSafeModeEnabled())
+	{
+		if (Index >= EventInfo.GetNumFields())
+		{
+			SetErrorInvalidFieldIndex(Index);
+			return *this;
+		}
+		const FTraceWriterEventFieldInfo& FieldInfo = EventInfo.GetField(Index);
+		if (!EnumHasAnyFlags(FieldInfo.GetType(), ETraceWriterFieldType::ArrayFlag) ||
+			EnumHasAnyFlags(FieldInfo.GetType(), ETraceWriterFieldType::ReferenceFlag | ETraceWriterFieldType::DefinitionIdFlag) ||
+			FieldInfo.GetType() == (ETraceWriterFieldType::Bool | ETraceWriterFieldType::ArrayFlag) || // not supported by protocol 7
+			FTraceWriter::GetByteSizeForFieldType(FieldInfo.GetType() & ~ETraceWriterFieldType::ArrayFlag) == 0)
+		{
+			SetErrorInvalidFieldType(Index);
+			return *this;
+		}
+		if (DataSize >= (1u << 16) || FieldInfo.GetSize() != 0)
+		{
+			SetErrorInvalidFieldData(Index);
+			return *this;
+		}
+	}
+	Fields[Index].SetData(Data, DataSize);
+	return *this;
+}
+
+////////////////////////////////////////////////////////////////////////////////////////////////////
+
+bool FTraceWriterEventBuilder::CheckValidPodField(uint32 Index, ETraceWriterFieldType FieldType, uint32 Size)
+{
+	if (Writer.IsSafeModeEnabled())
+	{
+		if (Index >= EventInfo.GetNumFields())
+		{
+			SetErrorInvalidFieldIndex(Index);
+			return false;
+		}
+		const FTraceWriterEventFieldInfo& FieldInfo = EventInfo.GetField(Index);
+		if (FieldInfo.GetSize() == 0 /*||
+			FTraceWriter::IsCompatibleType(FieldInfo.GetType(), FieldType)*/)
+		{
+			SetErrorInvalidFieldType(Index);
+			return false;
+		}
+		if (Size >= (1u << 16) || FieldInfo.GetSize() != Size)
+		{
+			SetErrorInvalidFieldData(Index);
+			return false;
+		}
+	}
+	return true;
+}
+
+////////////////////////////////////////////////////////////////////////////////////////////////////
+
+bool FTraceWriterEventBuilder::CheckValidExactFieldType(uint32 Index, ETraceWriterFieldType FieldType)
+{
+	if (Writer.IsSafeModeEnabled())
+	{
+		if (Index >= EventInfo.GetNumFields())
+		{
+			SetErrorInvalidFieldIndex(Index);
+			return false;
+		}
+		if (EventInfo.GetField(Index).GetType() != FieldType)
+		{
+			SetErrorInvalidFieldType(Index);
+			return false;
+		}
+	}
+	return true;
+}
+
+////////////////////////////////////////////////////////////////////////////////////////////////////
+
+void FTraceWriterEventBuilder::SetField(uint32 Index, const void* Data, uint32 DataSize)
+{
+	check(Index < uint32(Fields.Num()));
+	if (Writer.IsSafeModeEnabled())
+	{
+		const FTraceWriterEventFieldInfo& FieldInfo = EventInfo.GetField(Index);
+		if (DataSize >= (1u << 16) ||
+			(FieldInfo.GetSize() != 0 && FieldInfo.GetSize() != DataSize))
+		{
+			SetErrorInvalidFieldData(Index);
+			return;
+		}
+	}
+	Fields[Index].SetData(Data, DataSize);
+}
+
+////////////////////////////////////////////////////////////////////////////////////////////////////
+
+void FTraceWriterEventBuilder::End()
+{
+	// Write the event to the stream...
+	Writer.WriteEvent(*this);
+
+	// End() destroys the current EventBuilder!
+	bIsCompleted = true;
+	delete this;
+}
+
+////////////////////////////////////////////////////////////////////////////////////////////////////
+
+void FTraceWriterEventBuilder::SetErrorInvalidFieldIndex(uint32 Index)
+{
+	Writer.SetError(TEXT("Invalid field index!"));
+}
+
+////////////////////////////////////////////////////////////////////////////////////////////////////
+
+void FTraceWriterEventBuilder::SetErrorInvalidFieldType(uint32 Index)
+{
+	Writer.SetError(TEXT("Invalid field type!"));
+}
+
+////////////////////////////////////////////////////////////////////////////////////////////////////
+
+void FTraceWriterEventBuilder::SetErrorInvalidFieldData(uint32 Index)
+{
+	Writer.SetError(TEXT("Invalid field data!"));
+}
+
+////////////////////////////////////////////////////////////////////////////////////////////////////
+// FTraceWriter
+////////////////////////////////////////////////////////////////////////////////////////////////////
+
+FTraceWriter::FTraceWriter(IOutDataStream& InStream)
+	: Stream(InStream)
+{
+	GenerateGuid(&TraceGuid);
+	GenerateGuid(&SessionGuid);
+
+	EventInfos.SetNumZeroed(uint32(Private::WriterProtocol::EKnownEventUids::User), EAllowShrinking::No);
+}
+
+////////////////////////////////////////////////////////////////////////////////////////////////////
+
+FTraceWriter::~FTraceWriter()
+{
+	Reset();
+
+	if (EventBuffer)
+	{
+		FMemory::Free(EventBuffer);
+		EventBuffer = nullptr;
+	}
+	EventBufferSize = 0;
+
+	if (PacketBuffer)
+	{
+		FMemory::Free(PacketBuffer);
+		PacketBuffer = nullptr;
+	}
+}
+
+////////////////////////////////////////////////////////////////////////////////////////////////////
+
+void FTraceWriter::Reset()
+{
+	Stream.Close();
+
+	//TraceGuid
+	//SessionGuid
+
+	StartCycle = 0;
+	CycleFrequency = 0;
+	StartDateTime = 0.0;
+
+	NextSerial = 0;
+
+	UnknownEventId = 0;
+	NewTraceEventId = 0;
+	ThreadTimingEventId = 0;
+	ThreadInfoEventId = 0;
+	ThreadGroupBeginEventId = 0;
+	ThreadGroupEndEventId = 0;
+
+	Threads.Reset();
+	CurrentThreadId = ETransportTid::Events;
+
+	for (FTraceWriterEventInfo* EventInfo : EventInfos)
+	{
+		if (EventInfo)
+		{
+			delete EventInfo;
+		}
+	}
+	EventInfos.SetNum(uint32(Private::WriterProtocol::EKnownEventUids::User), EAllowShrinking::No);
+
+	if (PendingEventDeclaration)
+	{
+		delete PendingEventDeclaration;
+		PendingEventDeclaration = nullptr;
+	}
+
+	if (PendingEvent)
+	{
+		delete PendingEvent;
+		PendingEvent = nullptr;
+	}
+
+	//EventBuffer
+	//EventBufferSize
+	EventBufferDataSize = 0;
+	CompletedEventsDataSize = 0;
+
+	//PacketBuffer
+
+	LastError.Reset();
+}
+
+////////////////////////////////////////////////////////////////////////////////////////////////////
+
+ETransport FTraceWriter::GetTransportProtocolVersion() const
+{
+	return ETransport::TidPacketSync;
+}
+
+////////////////////////////////////////////////////////////////////////////////////////////////////
+
+uint8 FTraceWriter::GetTraceProtocolVersion() const
+{
+	return (uint32)Private::WriterProtocol::EProtocol::Id;
+}
+
+////////////////////////////////////////////////////////////////////////////////////////////////////
+
+void FTraceWriter::GenerateGuid(FTraceWriter::FTraceGuid* OutGuid)
+{
+	// This is not thread safe. Should only be accessed from the writer thread.
+	// This initialized the prng with the current timestamp. In theory two machines could initialize on the exact same time
+	// producing the same sequence of guids.
+	static uint64 State = FPlatformTime::Cycles64();
+
+	// L'Ecuyer, Pierre (1999). "Tables of Linear Congruential Generators of Different Sizes and Good Lattice Structure"
+	// corrected with errata
+	// Assuming m = 2e64
+	constexpr uint64 C = 0x369DEA0F31A53F85;
+	constexpr uint64 I = 1ull;
+
+	const uint64 TopBits = State * C + I;
+	const uint64 BottomBits = TopBits * C + I;
+	State = BottomBits;
+
+	*(uint64*)&OutGuid->Bits[0] = TopBits;
+	*(uint64*)&OutGuid->Bits[2] = BottomBits;
+
+	constexpr uint8 Version = 0x40; //Version 4, 4 bits
+	constexpr uint8 VersionMask = 0xf0;
+	constexpr uint8 Variant = 0x80; //Variant 1, 2 bits
+	constexpr uint8 VariantMask = 0xc0;
+
+	uint8* Octets = (uint8*)OutGuid;
+	Octets[6] = Version | (~VersionMask & Octets[6]); // Octet 9
+	Octets[8] = Variant | (~VariantMask & Octets[8]); // Octet 7
+}
+
+////////////////////////////////////////////////////////////////////////////////////////////////////
+
+uint64 FTraceWriter::GetTimeFrequency() const
+{
+	return FPlatformTime::SecondsToCycles64(1.0);
+}
+
+////////////////////////////////////////////////////////////////////////////////////////////////////
+
+uint64 FTraceWriter::GetTimeCycles64() const
+{
+	return FPlatformTime::Cycles64();
+}
+
+////////////////////////////////////////////////////////////////////////////////////////////////////
+
+double FTraceWriter::GetTimeSinceEpoch() const
+{
+	//return std::chrono::duration_cast<std::chrono::duration<double>>(std::chrono::system_clock::now().time_since_epoch()).count();
+	return FDateTime::Now().ToUnixTimestampDecimal();
+}
+
+////////////////////////////////////////////////////////////////////////////////////////////////////
+
+void FTraceWriter::Begin()
+{
+	// Create the default threads.
+	Threads.Reset();
+	Threads.Add({ 0, FAnsiString("NewEvents") });
+	Threads.Add({ 1, FAnsiString("Importants") });
+	static_assert(ETransportTid::Events == 0);
+	static_assert(ETransportTid::Importants == 1);
+	check(Threads.Num() == ETransportTid::Bias);
+
+	StartDateTime = GetTimeSinceEpoch();
+	StartCycle = GetTimeCycles64();
+	CycleFrequency = GetTimeFrequency();
+
+	constexpr uint16 GControlPort = 0;
+
+	// Handshake
+	struct FHandshake
+	{
+		uint32 Magic = '2' | ('C' << 8) | ('R' << 16) | ('T' << 24);
+		uint16 MetadataSize = uint16(MetadataSizeSum);
+
+		uint16 MetadataField0 = uint16(sizeof(ControlPort) | (ControlPortFieldId << 8));
+		uint16 ControlPort = uint16(GControlPort);
+
+		uint16 MetadataField1 = uint16(sizeof(FTraceGuid) | (SessionGuidFieldId << 8));
+		uint8 SessionGuid[16];	// Avoid padding
+
+		uint16 MetadataField2 = uint16(sizeof(FTraceGuid) | (TraceGuidFieldId << 8));
+		uint8 TraceGuid[16];	// Avoid padding
+
+		enum
+		{
+			MetadataSizeSum = 2 + 2 + 2 + 16 + 2 + 16,
+			Size = MetadataSizeSum + 4 + 2,
+			ControlPortFieldId = 0,
+			SessionGuidFieldId = 1,
+			TraceGuidFieldId = 2,
+		};
+	};
+	FHandshake Handshake;
+	FMemory::Memcpy(&Handshake.SessionGuid, &SessionGuid, sizeof(FTraceGuid));
+	FMemory::Memcpy(&Handshake.TraceGuid, &TraceGuid, sizeof(FTraceGuid));
+	int32 Ret = Stream.Write(&Handshake, FHandshake::Size);
+	if (Ret != FHandshake::Size)
+	{
+		SetError(TEXT("Failed to write trace header!"));
+		return;
+	}
+
+	// Write versions.
+	struct FVersions
+	{
+		uint8 TransportVersion = ETransport::TidPacketSync;
+		uint8 ProtocolVersion = Private::WriterProtocol::EProtocol::Id;
+	};
+	FVersions Versions;
+	Ret = Stream.Write(&Versions, sizeof(FVersions));
+	if (Ret != sizeof(FVersions))
+	{
+		SetError(TEXT("Failed to write trace header!"));
+		return;
+	}
+
+	// Describe initial events.
+	RegisterDefaultEvents();
+	WritePacket();
+
+	// Write $Trace.NewTrace event.
+	SetCurrentThread(static_cast<uint32>(ETransportTid::Importants));
+	WriteEvent(NewTraceEventId)
+		.Field(0, StartCycle)
+		.Field(1, CycleFrequency)
+		.Field(2, uint16(0x524d)) // Endian
+		.Field(3, uint8(sizeof(void*))) // PointerSize
+		.Field(4, StartDateTime)
+		.End();
+}
+
+////////////////////////////////////////////////////////////////////////////////////////////////////
+
+void FTraceWriter::End()
+{
+	WritePacket();
+	Stream.Close();
+}
+
+////////////////////////////////////////////////////////////////////////////////////////////////////
+
+void FTraceWriter::RegisterDefaultEvents()
+{
+	NewTraceEventId = DeclareEvent(ANSITEXTVIEW("$Trace"), ANSITEXTVIEW("NewTrace"), ETraceWriterEventFlags::ImportantNoSync)
+		.Field(ANSITEXTVIEW("StartCycle"), ETraceWriterFieldType::Uint64)
+		.Field(ANSITEXTVIEW("CycleFrequency"), ETraceWriterFieldType::Uint64)
+		.Field(ANSITEXTVIEW("Endian"), ETraceWriterFieldType::Uint16)
+		.Field(ANSITEXTVIEW("PointerSize"), ETraceWriterFieldType::Uint8)
+		.Field(ANSITEXTVIEW("StartDateTime"), ETraceWriterFieldType::Float64)
+		.End();
+
+	ThreadTimingEventId = DeclareEvent(ANSITEXTVIEW("$Trace"), ANSITEXTVIEW("ThreadTiming"), ETraceWriterEventFlags::NoSync)
+		.Field(ANSITEXTVIEW("BaseTimestamp"), ETraceWriterFieldType::Uint64)
+		.End();
+
+	ThreadInfoEventId = DeclareEvent(ANSITEXTVIEW("$Trace"), ANSITEXTVIEW("ThreadInfo"), ETraceWriterEventFlags::ImportantNoSync)
+		.Field(ANSITEXTVIEW("ThreadId"), ETraceWriterFieldType::Uint32)
+		.Field(ANSITEXTVIEW("SystemId"), ETraceWriterFieldType::Uint32)
+		.Field(ANSITEXTVIEW("SortHint"), ETraceWriterFieldType::Int32)
+		.Field(ANSITEXTVIEW("Name"), ETraceWriterFieldType::AnsiString)
+		.End();
+
+	ThreadGroupBeginEventId = DeclareEvent(ANSITEXTVIEW("$Trace"), ANSITEXTVIEW("ThreadGroupBegin"), ETraceWriterEventFlags::ImportantNoSync)
+		.Field(ANSITEXTVIEW("Name"), ETraceWriterFieldType::AnsiString)
+		.End();
+
+	ThreadGroupEndEventId = DeclareEvent(ANSITEXTVIEW("$Trace"), ANSITEXTVIEW("ThreadGroupEnd"), ETraceWriterEventFlags::ImportantNoSync)
+		.End();
+}
+
+////////////////////////////////////////////////////////////////////////////////////////////////////
+
+uint32 FTraceWriter::DeclareDiagnosticsSession2Event()
+{
+	return DeclareEvent(ANSITEXTVIEW("Diagnostics"), ANSITEXTVIEW("Session2"), ETraceWriterEventFlags::ImportantNoSync)
+		.Field(ANSITEXTVIEW("ConfigurationType"), ETraceWriterFieldType::Uint8)
+		.Field(ANSITEXTVIEW("TargetType"), ETraceWriterFieldType::Uint8)
+		.Field(ANSITEXTVIEW("Changelist"), ETraceWriterFieldType::Uint32)
+		.Field(ANSITEXTVIEW("InstanceId"), ETraceWriterFieldType::Uint32 | ETraceWriterFieldType::ArrayFlag)
+		.Field(ANSITEXTVIEW("Platform"), ETraceWriterFieldType::AnsiString)
+		.Field(ANSITEXTVIEW("AppName"), ETraceWriterFieldType::AnsiString)
+		.Field(ANSITEXTVIEW("ProjectName"), ETraceWriterFieldType::WideString)
+		.Field(ANSITEXTVIEW("Branch"), ETraceWriterFieldType::WideString)
+		.Field(ANSITEXTVIEW("BuildVersion"), ETraceWriterFieldType::WideString)
+		.Field(ANSITEXTVIEW("EngineVersion"), ETraceWriterFieldType::WideString)
+		.Field(ANSITEXTVIEW("CommandLine"), ETraceWriterFieldType::WideString)
+		.Field(ANSITEXTVIEW("VFSPaths"), ETraceWriterFieldType::AnsiString)
+		.End();
+}
+
+////////////////////////////////////////////////////////////////////////////////////////////////////
+
+uint32 FTraceWriter::GetOrRegisterUnknownEventId()
+{
+	if (UnknownEventId == 0)
+	{
+		UnknownEventId = DeclareEvent(ANSITEXTVIEW("TraceWriter"), ANSITEXTVIEW("UnknownEvent"), ETraceWriterEventFlags::ImportantNoSync)
+			.End();
+	}
+	return UnknownEventId;
+}
+
+////////////////////////////////////////////////////////////////////////////////////////////////////
+
+uint32 FTraceWriter::DeclareThread(FAnsiStringView ThreadName, uint32 SystemId, int32 SortHint)
+{
+	const FThreadInfo* FoundThread = Threads.FindByPredicate(
+		[&ThreadName](const FThreadInfo& ThreadInfo)
+		{
+			return ThreadName.Equals(ThreadInfo.Name, ESearchCase::IgnoreCase);
+		});
+	if (FoundThread)
+	{
+		return FoundThread->Id;
+	}
+
+	SetCurrentThreadImportants();
+
+	FThreadInfo& Thread = Threads.AddDefaulted_GetRef();
+	Thread.Id = uint32(Threads.Num()) - 1;
+	Thread.Name = FAnsiString(ThreadName);
+	Thread.SystemId = SystemId;
+	Thread.SortHint = SortHint;
+
+	WriteEvent(ThreadInfoEventId) // ImportantNoSync
+		.Field(0, Thread.Id)
+		.Field(1, Thread.SystemId)
+		.Field(2, Thread.SortHint)
+		.Field(3, Thread.Name)
+		.End();
+
+	return Thread.Id;
+}
+
+////////////////////////////////////////////////////////////////////////////////////////////////////
+
+void FTraceWriter::SetCurrentThread(uint32 ThreadId)
+{
+	if (ThreadId != CurrentThreadId)
+	{
+		// A packet can only contain events from a single thread.
+		WritePacket();
+
+		if ((EventBufferDataSize > 0) ||
+			(PendingEvent != nullptr) ||
+			(PendingEventDeclaration != nullptr))
+		{
+			SetError(TEXT("Cannot change the current thread while writing an event or an event declaration! Check if an End() call is missing."));
+			return;
+		}
+
+		CurrentThreadId = ThreadId;
+	}
+}
+
+////////////////////////////////////////////////////////////////////////////////////////////////////
+
+void FTraceWriter::WriteThreadTimingEvent(uint64 BaseTimestamp)
+{
+	WriteEvent(ThreadTimingEventId) // NoSync
+		.Field(0, BaseTimestamp)
+		.End();
+}
+
+////////////////////////////////////////////////////////////////////////////////////////////////////
+
+void FTraceWriter::WriteThreadGroupBeginEvent(FAnsiStringView GroupName)
+{
+	SetCurrentThreadImportants();
+	WriteEvent(ThreadGroupBeginEventId) // ImportantNoSync
+		.Field(0, GroupName)
+		.End();
+}
+
+////////////////////////////////////////////////////////////////////////////////////////////////////
+
+void FTraceWriter::WriteThreadGroupEndEvent()
+{
+	SetCurrentThreadImportants();
+	WriteEvent(ThreadGroupEndEventId) // ImportantNoSync
+		.End();
+}
+
+////////////////////////////////////////////////////////////////////////////////////////////////////
+
+FTraceWriterEventDeclarationBuilder& FTraceWriter::DeclareEvent(FAnsiStringView LoggerName, FAnsiStringView Name, ETraceWriterEventFlags Flags)
+{
+	check(PendingEventDeclaration == nullptr);
+
+	if (LoggerName.IsEmpty())
+	{
+		SetError(TEXT("Logger name cannot be empty!"));
+		LoggerName = ANSITEXTVIEW("Unknown");
+	}
+
+	if (Name.IsEmpty())
+	{
+		SetError(TEXT("Event name cannot be empty!"));
+		Name = ANSITEXTVIEW("Unknown");
+	}
+
+	if (EnumHasAnyFlags(Flags, ETraceWriterEventFlags::Important) &&
+		!EnumHasAnyFlags(Flags, ETraceWriterEventFlags::NoSync))
+	{
+		SetError(TEXT("Important events cannot be sync!"));
+		Flags = ETraceWriterEventFlags::ImportantNoSync;
+	}
+
+	FTraceWriterEventInfo* EventInfo = new FTraceWriterEventInfo(LoggerName, Name, Flags);
+	EventInfo->Id = static_cast<uint32>(EventInfos.Num());
+	EventInfos.Add(EventInfo);
+	
+	FTraceWriterEventDeclarationBuilder* Builder = new FTraceWriterEventDeclarationBuilder(*this, *EventInfo);
+	PendingEventDeclaration = Builder;
+	return *Builder;
+}
+
+////////////////////////////////////////////////////////////////////////////////////////////////////
+
+void FTraceWriter::WriteEventDeclaration(const FTraceWriterEventDeclarationBuilder& Builder)
+{
+	check(PendingEventDeclaration == &Builder);
+	PendingEventDeclaration = nullptr;
+
+	const FTraceWriterEventInfo& EventInfo = Builder.EventInfo;
+	check(EventInfos.Contains(&EventInfo));
+
+	SetCurrentThread(static_cast<uint32>(ETransportTid::Events));
+	BeginEvent();
+	WriteEventDeclaration(EventInfo);
+	EndEvent();
+}
+
+////////////////////////////////////////////////////////////////////////////////////////////////////
+
+void FTraceWriter::WriteEventDeclaration(const FTraceWriterEventInfo& EventInfo)
+{
+	// See the equivalent implementation in the TraceLog module:
+	//    TraceLog\Private\Trace\EventNode.cpp
+	//    FEventNode::Describe()
+
+	using namespace Private::WriterProtocol;
+
+	//////////////////////////////////////////////////
+
+	const uint32 FieldCount = EventInfo.GetNumFields();
+	if (FieldCount > 255)
+	{
+		SetError(TEXT("Too many fields!"));
+		return;
+	}
+
+	if (EventInfo.GetId() > uint32(EKnownEventUids::Max))
+	{
+		SetError(TEXT("Event id is too large!"));
+		return;
+	}
+
+	//////////////////////////////////////////////////
+	// Calculate the total size of names (strings).
+
+	uint32 NamesSize = 0;
+
+	if (EventInfo.GetLoggerName().Len() > 255)
+	{
+		SetError(TEXT("Logger name is too long!"));
+		return;
+	}
+	NamesSize += EventInfo.GetLoggerName().Len();
+
+	if (EventInfo.GetName().Len() > 255)
+	{
+		SetError(TEXT("Event name is too long!"));
+		return;
+	}
+	NamesSize += EventInfo.GetName().Len();
+
+	bool bMaybeHasAux = false;
+	for (uint32 FieldIndex = 0; FieldIndex < FieldCount; ++FieldIndex)
+	{
+		const FTraceWriterEventFieldInfo& FieldInfo = EventInfo.GetField(FieldIndex);
+		if (!EnumHasAnyFlags(FieldInfo.Type, ETraceWriterFieldType::DefinitionIdFlag))
+		{
+			if (FieldInfo.GetName().Len() > 255)
+			{
+				SetError(TEXT("Field name is too long!"));
+				return;
+			}
+			NamesSize += FieldInfo.GetName().Len();
+		}
+	}
+
+	//////////////////////////////////////////////////
+	// Compute the total event size.
+
+	uint32 EventSize =
+		  sizeof(FNewEventEvent)              // "NewEvent" event header
+		+ FieldCount * sizeof(FNewEventField) // fields
+		+ NamesSize;                          // strings
+	EventSize = (EventSize + 1) & ~1;         // align to 2 to match TraceLog
+
+	if (EventSize >= (1u << 16))
+	{
+		SetError(TEXT("NewEvent size is too large!"));
+		return;
+	}
+
+	const uint32 TotalEventSize =
+		  sizeof(FEventHeader) // event header (NoSync)
+		+ sizeof(uint16)       // event size
+		+ EventSize;           // event data
+
+	//////////////////////////////////////////////////
+	// Reserve buffer.
+	
+	uint8* Buffer = GetBuffer(TotalEventSize);
+	if (!Buffer)
+	{
+		SetError(TEXT("Failed to write new event!"));
+		return;
+	}
+
+	//////////////////////////////////////////////////
+	// Write event header(s).
+
+	FEventHeader* EventHeader = (FEventHeader*)Buffer;
+	EventHeader->Uid = EKnownEventUids::NewEvent;
+	Buffer += sizeof(FEventHeader);
+
+	*(uint16*)Buffer = static_cast<uint16>(EventSize);
+	Buffer += 2;
+
+	FNewEventEvent& NewEvent = *(FNewEventEvent*)Buffer;
+
+	NewEvent.EventUid = static_cast<uint16>(EventInfo.GetId());
+	NewEvent.FieldCount = static_cast<uint8>(EventInfo.GetNumFields());
+	NewEvent.Flags = GetProtocolEventFlags(EventInfo.GetFlags());
+	NewEvent.LoggerNameSize = static_cast<uint8>(EventInfo.GetLoggerName().Len());
+	NewEvent.EventNameSize = static_cast<uint8>(EventInfo.GetName().Len());
+	
+	// Write fields.
+	uint32 Offset = 0;
+	for (uint32 FieldIndex = 0; FieldIndex < FieldCount; ++FieldIndex)
+	{
+		const FTraceWriterEventFieldInfo& FieldInfo = EventInfo.GetField(FieldIndex);
+		FNewEventField& Field = *reinterpret_cast<FNewEventField*>(&NewEvent.Fields[FieldIndex]);
+		Field.FieldType = GetProtocolFieldFamily(FieldInfo.GetType());
+		Field.Unused = 0;
+		const uint8 TypeInfo = GetProtocolFieldType(FieldInfo.GetType());
+		switch (Field.FieldType)
+		{
+		case EFieldFamily::Regular:
+			Field.Regular.Offset = static_cast<uint16>(Offset);
+			Field.Regular.Size = FieldInfo.GetSize();
+			Field.Regular.TypeInfo = TypeInfo;
+			Field.Regular.NameSize = static_cast<uint8>(FieldInfo.GetName().Len());
+			break;
+		case EFieldFamily::Reference:
+			Field.Reference.Offset = static_cast<uint16>(Offset);
+			Field.Reference.RefUid = FieldInfo.GetRefUid();
+			Field.Reference.TypeInfo = TypeInfo;
+			Field.Reference.NameSize = static_cast<uint8>(FieldInfo.GetName().Len());
+			break;
+		case EFieldFamily::DefinitionId:
+			Field.DefinitionId.Offset = static_cast<uint16>(Offset);
+			Field.DefinitionId.Unused1 = 0;
+			Field.DefinitionId.Unused2 = 0;
+			Field.DefinitionId.TypeInfo = TypeInfo;
+			break;
+		default:
+			SetError(TEXT("Invalid field type!"));
+			return;
+		}
+		Offset += FieldInfo.GetSize();
+	}
+
+	// Write names
+	uint8* NamesCursor = (uint8*)(NewEvent.Fields + NewEvent.FieldCount);
+	auto WriteName = [&NamesCursor](const FAnsiString& String)
+		{
+			FMemory::Memcpy(NamesCursor, *String, String.Len());
+			NamesCursor += String.Len();
+		};
+	WriteName(EventInfo.GetLoggerName());
+	WriteName(EventInfo.GetName());
+	for (uint32 FieldIndex = 0; FieldIndex < FieldCount; ++FieldIndex)
+	{
+		const FTraceWriterEventFieldInfo& FieldInfo = EventInfo.GetField(FieldIndex);
+		if (!EnumHasAnyFlags(FieldInfo.Type, ETraceWriterFieldType::DefinitionIdFlag))
+		{
+			WriteName(FieldInfo.GetName());
+		}
+	}
+
+	//////////////////////////////////////////////////
+	// Commit the written buffer.
+
+	AdvanceBuffer(TotalEventSize);
+}
+
+////////////////////////////////////////////////////////////////////////////////////////////////////
+
+FTraceWriterEventBuilder& FTraceWriter::WriteEvent(uint32 EventId)
+{
+	check(PendingEvent == nullptr);
+
+	if (EventId >= static_cast<uint32>(EventInfos.Num()))
+	{
+		SetError(TEXT("Invalid event id!"));
+		EventId = GetOrRegisterUnknownEventId();
+	}
+
+	check(EventId < static_cast<uint32>(EventInfos.Num()));
+	check(EventInfos[EventId] != nullptr);
+
+	FTraceWriterEventBuilder* Builder = new FTraceWriterEventBuilder(*this, *EventInfos[EventId]);
+	PendingEvent = Builder;
+	return *Builder;
+}
+
+////////////////////////////////////////////////////////////////////////////////////////////////////
+
+void FTraceWriter::WriteEvent(const FTraceWriterEventBuilder& Builder)
+{
+	using namespace Private::WriterProtocol;
+
+	check(PendingEvent == &Builder);
+	PendingEvent = nullptr;
+
+	const FTraceWriterEventInfo& EventInfo = Builder.EventInfo;
+	check(EventInfos.Contains(&EventInfo));
+
+	const bool bMaybeHasAux = EnumHasAnyFlags(EventInfo.GetFlags(), ETraceWriterEventFlags::MaybeHasAux);
+	const bool bIsImportant = EnumHasAnyFlags(EventInfo.GetFlags(), ETraceWriterEventFlags::Important);
+	bool bIsNoSync = EnumHasAnyFlags(EventInfo.GetFlags(), ETraceWriterEventFlags::NoSync);
+
+	if (bIsImportant)
+	{
+		if (!bIsNoSync)
+		{
+			SetError(TEXT("Important events must be NoSync!"));
+			bIsNoSync = true;
+		}
+		if (CurrentThreadId != uint32(ETransportTid::Importants))
+		{
+			SetError(TEXT("Important events can only be emitted on the Importants thread!"));
+			SetCurrentThread(static_cast<uint32>(ETransportTid::Importants));
+		}
+	}
+	else
+	{
+		if (CurrentThreadId == uint32(ETransportTid::Importants))
+		{
+			SetError(TEXT("Non-important events can not be emitted on the Importants thread!"));
+			SetCurrentThread(uint32(ETransportTid::Bias));
+		}
+	}
+
+	BeginEvent();
+	WriteEvent(Builder, EventInfo);
+	EndEvent();
+
+	if (bMaybeHasAux)
+	{
+		// Write the AuxData events.
+		const uint32 FieldCount = EventInfo.GetNumFields();
+		for (uint32 FieldIndex = 0; FieldIndex < FieldCount; ++FieldIndex)
+		{
+			const FTraceWriterEventFieldInfo& FieldInfo = EventInfo.GetField(FieldIndex);
+			uint32 FieldSize = FieldInfo.GetSize();
+			if (FieldSize == 0)
+			{
+				const FTraceWriterFieldData& FieldData = Builder.Fields[FieldIndex];
+				if (FieldData.GetSize() > 0)
+				{
+					WriteAuxData(FieldIndex, FieldData.GetData(), FieldData.GetSize(), bIsImportant);
+				}
+			}
+		}
+
+		// Write the AuxDataTerminal event.
+		BeginEvent();
+		uint8* Buffer = GetBuffer(1);
+		if (Buffer)
+		{
+			*Buffer = bIsImportant ?
+				uint8(EKnownEventUids::AuxDataTerminal) : // no _UidShift !!!
+				uint8(EKnownEventUids::AuxDataTerminal << EKnownEventUids::_UidShift);
+			AdvanceBuffer(1);
+		}
+		else
+		{
+			SetError(TEXT("Failed to write aux terminal event!"));
+		}
+		EndEvent();
+	}
+}
+
+////////////////////////////////////////////////////////////////////////////////////////////////////
+
+void FTraceWriter::WriteEvent(const FTraceWriterEventBuilder& Builder, const FTraceWriterEventInfo& EventInfo)
+{
+	// See the equivalent implementation in the TraceLog module:
+	//    TraceLog\Private\Trace\*.cpp
+
+	using namespace Private::WriterProtocol;
+
+	//////////////////////////////////////////////////
+
+	const uint32 FieldCount = EventInfo.GetNumFields();
+	check(FieldCount <= 255);
+	check(FieldCount == Builder.Fields.Num());
+
+	//////////////////////////////////////////////////
+	// Compute event size.
+
+	uint32 EventSize = 0;
+	uint32 AuxSize = 0;
+	for (uint32 FieldIndex = 0; FieldIndex < FieldCount; ++FieldIndex)
+	{
+		const FTraceWriterEventFieldInfo& FieldInfo = EventInfo.GetField(FieldIndex);
+		EventSize += FieldInfo.GetSize();
+		if (FieldInfo.GetSize() == 0)
+		{
+			const FTraceWriterFieldData& FieldData = Builder.Fields[FieldIndex];
+			if (FieldData.GetSize() > 0)
+			{
+				AuxSize += sizeof(FAuxHeader) + FieldData.GetSize();
+			}
+		}
+	}
+	if (EnumHasAllFlags(EventInfo.GetFlags(), ETraceWriterEventFlags::MaybeHasAux))
+	{
+		AuxSize += 1; // aux data terminal
+	}
+
+	if (EventSize >= (1u << 16))
+	{
+		SetError(TEXT("Event size is too large!"));
+		return;
+	}
+
+	uint32 TotalEventSize = 0;
+	if (EnumHasAnyFlags(EventInfo.GetFlags(), ETraceWriterEventFlags::Important))
+	{
+		TotalEventSize = sizeof(FImportantEventHeader); // event header
+	}
+	else if (EnumHasAnyFlags(EventInfo.GetFlags(), ETraceWriterEventFlags::NoSync))
+	{
+		TotalEventSize = sizeof(FEventHeader); // event header
+	}
+	else // Sync
+	{
+		TotalEventSize = sizeof(FEventHeaderSync); // event header
+	}
+	TotalEventSize += EventSize; // event data
+
+	//////////////////////////////////////////////////
+	// Reserve buffer.
+
+	uint8* Buffer = GetBuffer(TotalEventSize + AuxSize);
+	if (!Buffer)
+	{
+		SetError(TEXT("Failed to write event!"));
+		return;
+	}
+
+	//////////////////////////////////////////////////
+
+	if (EnumHasAnyFlags(EventInfo.GetFlags(), ETraceWriterEventFlags::Important))
+	{
+		if (EventSize + AuxSize >= (1u << 16))
+		{
+			SetError(TEXT("Important event size is too large!"));
+			return;
+		}
+
+		check(CurrentThreadId == uint32(ETransportTid::Importants));
+		uint16 Uid16 = uint16(EventInfo.GetId()); // no Flag_TwoByteUid bit !!!
+		FImportantEventHeader* EventHeader = (FImportantEventHeader*)Buffer;
+		EventHeader->Uid = uint16(EventInfo.GetId());
+		EventHeader->Size = uint16(EventSize + AuxSize); // includes size of aux data !!!
+		Buffer += sizeof(FImportantEventHeader);
+	}
+	else if (EnumHasAnyFlags(EventInfo.GetFlags(), ETraceWriterEventFlags::NoSync))
+	{
+		FEventHeader* EventHeader = (FEventHeader*)Buffer;
+		EventHeader->Uid = uint16(EventInfo.GetId() << EKnownEventUids::_UidShift) | EKnownEventUids::Flag_TwoByteUid;
+		Buffer += sizeof(FEventHeader);
+	}
+	else // Sync
+	{
+		FEventHeaderSync* EventHeader = (FEventHeaderSync*)Buffer;
+		FMemory::Memcpy(((uint8*)EventHeader->Data) - 3, &NextSerial, sizeof(NextSerial)); // FEventHeaderSync::SerialHigh,SerialLow
+		++NextSerial;
+		EventHeader->Uid = uint16(EventInfo.GetId() << EKnownEventUids::_UidShift) | EKnownEventUids::Flag_TwoByteUid;
+		Buffer += sizeof(FEventHeaderSync);
+	}
+
+	for (uint32 FieldIndex = 0; FieldIndex < FieldCount; ++FieldIndex)
+	{
+		const FTraceWriterEventFieldInfo& FieldInfo = EventInfo.GetField(FieldIndex);
+		uint32 FieldSize = FieldInfo.GetSize();
+		if (FieldSize != 0)
+		{
+			const FTraceWriterFieldData& FieldData = Builder.Fields[FieldIndex];
+			if (FieldData.GetSize() == FieldSize)
+			{
+				FMemory::Memcpy(Buffer, FieldData.GetData(), FieldData.GetSize());
+			}
+			else
+			{
+				if (FieldData.GetSize() != 0)
+				{
+					SetError(TEXT("Event field data does not match the field declaration!"));
+				}
+				FMemory::Memzero(Buffer, FieldSize);
+			}
+			Buffer += FieldSize;
+		}
+	}
+
+	//////////////////////////////////////////////////
+	// Commit the written buffer.
+
+	AdvanceBuffer(TotalEventSize);
+}
+
+////////////////////////////////////////////////////////////////////////////////////////////////////
+
+void FTraceWriter::WriteAuxData(uint32 Index, const void* Data, uint32 Size, bool bIsImportant)
+{
+	// Technically, this could be FAuxHeader::SizeLimit (512 KiB), but in that case, it will exceed
+	// the maximum size for a packet. The current packet size (even for uncompressed packets) is limited
+	// to uint16 (see FTidPacketEncoded and FTidPacket).
+	constexpr uint32 MaxAuxEventSize = 1 << 15;
+
+	while (Size > 0)
+	{
+		const uint32 SegmentSize = FMath::Min(MaxAuxEventSize, Size);
+
+		BeginEvent();
+		WriteAuxDataSegment(Index, Data, SegmentSize, bIsImportant);
+		EndEvent();
+
+		Data = ((const uint8*)Data) + SegmentSize;
+		Size -= SegmentSize;
+	}
+}
+
+////////////////////////////////////////////////////////////////////////////////////////////////////
+
+void FTraceWriter::WriteAuxDataSegment(uint32 Index, const void* SegmentData, uint32 SegmentSize, bool bIsImportant)
+{
+	using namespace Private::WriterProtocol;
+
+	uint32 Size = sizeof(FAuxHeader) + SegmentSize;
+
+	uint8* Buffer = GetBuffer(Size + 1);
+	if (!Buffer)
+	{
+		SetError(TEXT("Failed to write aux event!"));
+		return;
+	}
+
+	// Write header
+	uint32 Pack = SegmentSize << FAuxHeader::SizeShift;
+	Pack |= Index << FAuxHeader::FieldShift;
+	FMemory::Memcpy(Buffer, &Pack, sizeof(uint32)); /* FAuxHeader::Pack */
+	*Buffer = bIsImportant ?  
+		uint8(EKnownEventUids::AuxData) : // no _UidShift !!!
+		uint8(EKnownEventUids::AuxData << EKnownEventUids::_UidShift); /* FAuxHeader::Uid */
+	Buffer += sizeof(FAuxHeader);
+
+	// Write data
+	FMemory::Memcpy(Buffer, SegmentData, SegmentSize);
+
+	AdvanceBuffer(Size);
+}
+
+////////////////////////////////////////////////////////////////////////////////////////////////////
+
+uint16 FTraceWriter::GetByteSizeForFieldType(ETraceWriterFieldType Type)
+{
+	if (EnumHasAnyFlags(Type, ETraceWriterFieldType::ArrayFlag))
+	{
+		return 0; // arrays have variable size
+	}
+
+	switch (Type & ETraceWriterFieldType::IndexMask)
+	{
+	case ETraceWriterFieldType::Bool:
+	case ETraceWriterFieldType::Uint8:
+	case ETraceWriterFieldType::Int8:
+		return 1;
+
+	case ETraceWriterFieldType::Uint16:
+	case ETraceWriterFieldType::Int16:
+		return 2;
+
+	case ETraceWriterFieldType::Uint32:
+	case ETraceWriterFieldType::Int32:
+	case ETraceWriterFieldType::Float32: // float
+		return 4;
+
+	case ETraceWriterFieldType::Uint64:
+	case ETraceWriterFieldType::Int64:
+	case ETraceWriterFieldType::Float64: // double
+	case ETraceWriterFieldType::Pointer:
+		return 8;
+
+	default:
+		return 0; // unknown or variable size
+	}
+}
+
+////////////////////////////////////////////////////////////////////////////////////////////////////
+
+uint8 FTraceWriter::GetProtocolFieldFamily(ETraceWriterFieldType Type)
+{
+	// See UE::Trace::Protocol6::EFieldFamily (Trace\Detail\Protocols\Protocol6.h in TraceLog module).
+	// { Regular, Reference, DefinitionId }
+
+	using namespace Private::WriterProtocol;
+
+	if ((Type & ETraceWriterFieldType::FlagsMask) == ETraceWriterFieldType::None ||
+		(Type & ETraceWriterFieldType::FlagsMask) == ETraceWriterFieldType::ArrayFlag)
+	{
+		return EFieldFamily::Regular;
+	}
+	if ((Type & ETraceWriterFieldType::FlagsMask) == ETraceWriterFieldType::ReferenceFlag)
+	{
+		return EFieldFamily::Reference;
+	}
+	if ((Type & ETraceWriterFieldType::FlagsMask) == ETraceWriterFieldType::DefinitionIdFlag)
+	{
+		return EFieldFamily::DefinitionId;
+	}
+	return 0xFF; // error
+}
+
+////////////////////////////////////////////////////////////////////////////////////////////////////
+
+uint8 FTraceWriter::GetProtocolFieldType(ETraceWriterFieldType Type)
+{
+	// See UE::Trace::Protocol0::EFieldType (Trace\Detail\Protocols\Protocol0.h in TraceLog module).
+
+	using namespace Private::WriterProtocol;
+
+	if (EnumHasAnyFlags(Type, ETraceWriterFieldType::ArrayFlag))
+	{
+		return uint8(EFieldType::Array);
+	}
+
+	if (EnumHasAnyFlags(Type, ETraceWriterFieldType::ReferenceFlag))
+	{
+		//switch (Type & ETraceWriterFieldType::IndexMask)
+		//{
+		//	case ETraceWriterFieldType::Reference8:  return uint8(EFieldType::Uint8);
+		//	case ETraceWriterFieldType::Reference16: return uint8(EFieldType::Uint16);
+		//	case ETraceWriterFieldType::Reference32: return uint8(EFieldType::Uint32);
+		//	case ETraceWriterFieldType::Reference64: return uint8(EFieldType::Uint64);
+		//}
+		return 0;
+	}
+
+	if (EnumHasAnyFlags(Type, ETraceWriterFieldType::DefinitionIdFlag))
+	{
+		//switch (Type & ETraceWriterFieldType::IndexMask)
+		//{
+		//	case ETraceWriterFieldType::DefinitionId8:  return uint8(EFieldType::Uint8);
+		//	case ETraceWriterFieldType::DefinitionId16: return uint8(EFieldType::Uint16);
+		//	case ETraceWriterFieldType::DefinitionId32: return uint8(EFieldType::Uint32);
+		//	case ETraceWriterFieldType::DefinitionId64: return uint8(EFieldType::Uint64);
+		//}
+		return 0;
+	}
+
+	switch (Type & ETraceWriterFieldType::IndexMask)
+	{
+		case ETraceWriterFieldType::None:        return 0;
+		case ETraceWriterFieldType::Bool:        return uint8(EFieldType::Bool); // == EFieldType::Uint8 == 0
+		case ETraceWriterFieldType::Uint8:       return uint8(EFieldType::Uint8); // == 0
+		case ETraceWriterFieldType::Uint16:      return uint8(EFieldType::Uint16);
+		case ETraceWriterFieldType::Uint32:      return uint8(EFieldType::Uint32);
+		//case ETraceWriterFieldType::Uint32_7bit: return 0; // not supported in Protocol7
+		case ETraceWriterFieldType::Uint64:      return uint8(EFieldType::Uint64);
+		//case ETraceWriterFieldType::Uint64_7bit: return 0; // not supported in Protocol7
+		case ETraceWriterFieldType::Int8:        return uint8(EFieldType::Int8);
+		case ETraceWriterFieldType::Int16:       return uint8(EFieldType::Int16);
+		case ETraceWriterFieldType::Int32:       return uint8(EFieldType::Int32);
+		//case ETraceWriterFieldType::Int32_7bit:  return 0; // not supported in Protocol7
+		case ETraceWriterFieldType::Int64:       return uint8(EFieldType::Int64);
+		//case ETraceWriterFieldType::Int64_7bit:  return 0; // not supported in Protocol7
+		case ETraceWriterFieldType::Float32:     return uint8(EFieldType::Float32);
+		case ETraceWriterFieldType::Float64:     return uint8(EFieldType::Float64);
+		case ETraceWriterFieldType::AnsiString:  return uint8(EFieldType::AnsiString);
+		case ETraceWriterFieldType::WideString:  return uint8(EFieldType::WideString);
+		//case ETraceWriterFieldType::Utf8String:  return 0; // not supported in Protocol7
+		case ETraceWriterFieldType::Pointer:     return uint8(EFieldType::Pointer); // == EFieldType::Uint64
+		default:
+			return 0;
+	}
+}
+
+////////////////////////////////////////////////////////////////////////////////////////////////////
+
+ETraceWriterFieldType FTraceWriter::ConvertFieldType(const UE::Trace::IAnalyzer::FEventFieldInfo& FieldInfo)
+{
+	ETraceWriterFieldType WriterFieldType = ETraceWriterFieldType::None;
+
+	using EType = UE::Trace::IAnalyzer::FEventFieldInfo::EType;
+	const EType Type = FieldInfo.GetType();
+	switch (Type)
+	{
+		case EType::Integer:
+			if (FieldInfo.IsSigned())
+			{
+				switch (FieldInfo.GetTypeSize())
+				{
+					case 1: WriterFieldType = ETraceWriterFieldType::Int8; break;
+					case 2: WriterFieldType = ETraceWriterFieldType::Int16; break;
+					case 4: WriterFieldType = ETraceWriterFieldType::Int32; break;
+					case 8: WriterFieldType = ETraceWriterFieldType::Int64; break;
+				}
+			}
+			else
+			{
+				switch (FieldInfo.GetTypeSize())
+				{
+					case 1: WriterFieldType = ETraceWriterFieldType::Uint8; break;
+					case 2: WriterFieldType = ETraceWriterFieldType::Uint16; break;
+					case 4: WriterFieldType = ETraceWriterFieldType::Uint32; break;
+					case 8: WriterFieldType = ETraceWriterFieldType::Uint64; break;
+				}
+			}
+			break;
+
+		case EType::Float:
+			if (FieldInfo.GetTypeSize() == 4)
+			{
+				WriterFieldType = ETraceWriterFieldType::Float32;
+			}
+			else
+			{
+				check(FieldInfo.GetTypeSize() == 8);
+				WriterFieldType = ETraceWriterFieldType::Float64;
+			}
+			break;
+
+		case EType::AnsiString: WriterFieldType = ETraceWriterFieldType::AnsiString; break;
+		case EType::WideString: WriterFieldType = ETraceWriterFieldType::WideString; break;
+
+		case EType::Reference8:  WriterFieldType = ETraceWriterFieldType::ReferenceFlag | ETraceWriterFieldType::Uint8;  break;
+		case EType::Reference16: WriterFieldType = ETraceWriterFieldType::ReferenceFlag | ETraceWriterFieldType::Uint16; break;
+		case EType::Reference32: WriterFieldType = ETraceWriterFieldType::ReferenceFlag | ETraceWriterFieldType::Uint32; break;
+		case EType::Reference64: WriterFieldType = ETraceWriterFieldType::ReferenceFlag | ETraceWriterFieldType::Uint64; break;
+	}
+
+	if (FieldInfo.IsArray())
+	{
+		EnumAddFlags(WriterFieldType, ETraceWriterFieldType::ArrayFlag);
+	}
+
+	return WriterFieldType;
+}
+
+////////////////////////////////////////////////////////////////////////////////////////////////////
+
+uint8* FTraceWriter::GetBuffer(uint32 RequiredDataSize)
+{
+	const uint32 TotalDataSize = MaxPacketHeaderSize + EventBufferDataSize + RequiredDataSize;
+	if (TotalDataSize > EventBufferSize)
+	{
+		constexpr uint32 PageSize = 64 * 1024;
+		EventBufferSize = (TotalDataSize + PageSize - 1) & ~(PageSize - 1);
+		EventBuffer = (uint8*)FMemory::Realloc(EventBuffer, EventBufferSize, 0);
+	}
+	return EventBuffer + MaxPacketHeaderSize + EventBufferDataSize;
+}
+
+////////////////////////////////////////////////////////////////////////////////////////////////////
+
+void FTraceWriter::AdvanceBuffer(uint32 DataSize)
+{
+	EventBufferDataSize += DataSize;
+	check(MaxPacketHeaderSize + EventBufferDataSize <= EventBufferSize);
+}
+
+////////////////////////////////////////////////////////////////////////////////////////////////////
+
+void FTraceWriter::BeginEvent()
+{
+}
+
+////////////////////////////////////////////////////////////////////////////////////////////////////
+
+void FTraceWriter::EndEvent()
+{
+	if (EventBufferDataSize > MaxDecodedBufferSize)
+	{
+		// Not enough packet space to include the current event.
+		// Write a packet with the previous event(s) first.
+		// The current event remains pending in the EventBuffer.
+		WritePacket();
+		check(EventBufferDataSize <= MaxDecodedBufferSize);
+	}
+
+	CompletedEventsDataSize = EventBufferDataSize;
+
+	if (CompletedEventsDataSize >= FlushPacketThreshold)
+	{
+		WritePacket();
+	}
+}
+
+////////////////////////////////////////////////////////////////////////////////////////////////////
+
+void FTraceWriter::WritePacket()
+{
+	if (CompletedEventsDataSize == 0)
+	{
+		return;
+	}
+
+	using namespace UE::Trace::Private;
+
+	uint8* Data = EventBuffer + MaxPacketHeaderSize;
+	uint32 Size = CompletedEventsDataSize;
+	check(Size <= MaxDecodedBufferSize); // PacketSize and DecodedSize are uint16
+
+	bool bWriteUncompressedData = false;
+
+	// Smaller buffers usually aren't redundant enough to benefit from being
+	// compressed. They often end up being larger.
+	if (Size > MinEncodedBufferSize && Size <= MaxEncodedBufferSize)
+	{
+		if (PacketBuffer == nullptr)
+		{
+			PacketBuffer = (uint8*)FMemory::Malloc(MaxPacketBufferSize, 0);
+		}
+		FTidPacketEncoded& Packet = *(FTidPacketEncoded*)PacketBuffer;
+		const int32 EncodedSize = Encode(Data, Size, Packet.Data, MaxEncodedBufferSize + EncodingOverhead);
+		if (EncodedSize > 0 && uint32(EncodedSize) < Size)
+		{
+			Packet.PacketSize = uint16(sizeof(FTidPacketEncoded) + EncodedSize);
+			Packet.ThreadId = FTidPacketBase::EncodedMarker;
+			Packet.ThreadId |= uint16(CurrentThreadId & FTidPacketBase::ThreadIdMask);
+			Packet.DecodedSize = uint16(Size);
+			Stream.Write(&Packet, Packet.PacketSize);
+		}
+		else
+		{
+			bWriteUncompressedData = true;
+		}
+	}
+	else
+	{
+		bWriteUncompressedData = true;
+	}
+
+	if (bWriteUncompressedData)
+	{
+		Data -= sizeof(FTidPacket);
+		Size += sizeof(FTidPacket);
+		FTidPacket& Packet = *(FTidPacket*)Data;
+		Packet.PacketSize = uint16(Size);
+		Packet.ThreadId = uint16(CurrentThreadId & FTidPacketBase::ThreadIdMask);
+		Stream.Write(Data, Size);
+	}
+
+	check(EventBufferDataSize >= CompletedEventsDataSize);
+	EventBufferDataSize -= CompletedEventsDataSize;
+	if (EventBufferDataSize > 0)
+	{
+		FMemory::Memmove(EventBuffer + MaxPacketHeaderSize, EventBuffer + MaxPacketHeaderSize + CompletedEventsDataSize, EventBufferDataSize);
+	}
+
+	CompletedEventsDataSize = 0;
+}
+
+////////////////////////////////////////////////////////////////////////////////////////////////////
+
+void FTraceWriter::SetError(FString&& InError)
+{
+	Errors.Push(InError);
+}
+
+////////////////////////////////////////////////////////////////////////////////////////////////////
+
+const FString& FTraceWriter::GetLastError() const
+{
+	if (Errors.Num() > 0)
+	{
+		LastError = Errors.Pop();
+	}
+	else
+	{
+		LastError.Reset();
+	}
+	return LastError;
+}
+
+////////////////////////////////////////////////////////////////////////////////////////////////////
+
+} // namespace UE::Trace
diff --git a/Engine/Source/Developer/TraceAnalysis/Public/Trace/DataStream.h b/Engine/Source/Developer/TraceAnalysis/Public/Trace/DataStream.h
--- a/Engine/Source/Developer/TraceAnalysis/Public/Trace/DataStream.h
+++ b/Engine/Source/Developer/TraceAnalysis/Public/Trace/DataStream.h
@@ -13,8 +13,8 @@
 class IFileHandle;
 class FSocket;
 
-namespace UE {
-namespace Trace {
+namespace UE::Trace
+{
 
 class IInDataStream
 {
@@ -22,25 +22,26 @@ public:
 	virtual ~IInDataStream() = default;
 
 	/**
-	 * Read bytes from the stream.
-	 * @param Data Pointer to buffer to read into
-	 * @param Size Maximum size that can be read in bytes
-	 * @return Number of bytes read from the stream. Zero indicates end of stream and negative values indicate errors.
+	 * Reads bytes from the stream.
+	 *
+	 * @param Data - The pointer to buffer to read into
+	 * @param Size - The maximum size that can be read in bytes
+	 * @returns the number of bytes read from the stream. Zero indicates end of stream and negative values indicate errors.
 	 */
 	virtual int32 Read(void* Data, uint32 Size) = 0;
 
 	/**
-	 * Close the stream. Reading from a closed stream
-	 * is considered an error.
+	 * Closes the stream.
+	 * Reading from a closed stream is considered an error.
 	 */
 	virtual void Close() {}
 
 	/**
-	 * Query if the stream is ready to read. Some streams may need to
-	 * establish the data stream before reading can begin. A stream may not
-	 * block indefinitely.
+	 * Waits for the the stream to be ready to read from.
+	 * Some streams may need to establish the data stream before reading can begin.
+	 * A stream may not block indefinitely.
 	 *
-	 * @return if the stream is ready to be read from
+	 * @returns if the stream is ready to be read from
 	 */
 	virtual bool WaitUntilReady() { return true; }
 };
@@ -57,9 +58,8 @@ public:
 	/*
 	* Open the file.
 	*
-	* @param Path The path to the file.
-	*
-	* @return True if the file was opened successfully.
+	* @param Path - The path to the file
+	* @returns true if the file was opened successfully
 	*/
 	UE_API bool Open(const TCHAR* Path);
 
@@ -72,8 +72,8 @@ private:
 };
 
 /**
- * Creates a stream to directly consume a trace stream from the tracing application. Sets up
- * a listening socket and stream is not considered ready until a connection is made.
+ * Creates a stream to directly consume a trace stream from the tracing application.
+ * Sets up a listening socket and stream is not considered ready until a connection is made.
  * If a valid file writer is provided during construction, the trace data will also be written to disk.
  */
 class FDirectSocketStream : public IInDataStream, public FRunnable
@@ -84,19 +84,19 @@ public:
 	UE_API virtual ~FDirectSocketStream() override;
 
 	/**
-	 * Initiates listening sockets. Must be called before attempting to read from
-	 * the stream.
-	 * @param Port The specific port number to connect to or 0 to use any available port
-	 * @return The actual port number used for listening
+	 * Initiates listening sockets.
+	 * Must be called before attempting to read from the stream.
+	 * @param Port - The specific port number to connect to or 0 to use any available port
+	 * @returns the actual port number used for listening
 	 */
 	UE_API uint16 StartListening(uint16 Port = 0);
 
 	/**
-	 * Bind this instance to an existing connection that has already been established.
-	 * Using this method allows you to service connections off a single listen port, then 
+	 * Binds this instance to an existing connection that has already been established.
+	 * Using this method allows you to service connections off a single listen port, then
 	 * create instances of FDirectSocketStream for each connection.
 	 * @param ExistingConnectedSocket is the existing socket that will transfer its ownership to this instance, even if the function fails.
-	 * @return true if the binding was successful. Note: Even if it returns false, you will have lost ownership of ExistingConnectedSocket.
+	 * @returns true if the binding was successful. Note: Even if it returns false, you will have lost ownership of ExistingConnectedSocket.
 	 */
 	UE_API bool BindToExistingConnection(FSocket&& ExistingConnectedSocket);
 
@@ -110,7 +110,7 @@ private:
 		MaxQueuedConnections = 4,	// Size of connection queue.
 	};
 
-	// IInStream interface
+	// IInDataStream interface
 	UE_API virtual bool WaitUntilReady() override;
 	UE_API virtual void Close() override;
 
@@ -129,7 +129,6 @@ private:
 	FEvent* ConnectionEvent;
 };
 
-} // namespace Trace
-} // namespace UE
+} // namespace UE::Trace
 
 #undef UE_API
diff --git a/Engine/Source/Developer/TraceAnalysis/Public/Trace/OutDataStream.h b/Engine/Source/Developer/TraceAnalysis/Public/Trace/OutDataStream.h
new file mode 100644
--- /dev/null
+++ b/Engine/Source/Developer/TraceAnalysis/Public/Trace/OutDataStream.h
@@ -0,0 +1,72 @@
+// Copyright Epic Games, Inc. All Rights Reserved.
+
+#pragma once
+
+#include "CoreTypes.h"
+#include "HAL/FileManagerGeneric.h"
+#include "Templates/UniquePtr.h"
+
+#define UE_API TRACEANALYSIS_API
+
+class IFileHandle;
+
+namespace UE::Trace
+{
+
+class IOutDataStream
+{
+public:
+	virtual ~IOutDataStream() = default;
+
+	/**
+	 * Writes bytes to the stream.
+	 *
+	 * @param Data - The pointer to the data buffer to read from and write to the stream
+	 * @param Size - The number of bytes to read from the Data buffer and to write to the stream
+	 * @returns the number of bytes written to the stream. Negative values indicate errors.
+	 */
+	virtual int32 Write(void* Data, uint32 Size) = 0;
+
+	/**
+	 * Closes the stream.
+	 * Writing to a closed stream is considered an error.
+	 */
+	virtual void Close() {}
+
+	/**
+	 * Waits for the stream to be ready to write to.
+	 * Some streams may need to establish the data stream before writing can begin.
+	 * A stream may not block indefinitely.
+	 *
+	 * @returns if the stream is ready to be write to
+	 */
+	virtual bool WaitUntilReady() { return true; }
+};
+
+/*
+ * An implementation of IOutDataStream that writes to a file on disk.
+ */
+class FFileOutDataStream : public IOutDataStream
+{
+public:
+	UE_API FFileOutDataStream();
+	UE_API virtual ~FFileOutDataStream() override;
+
+	/*
+	* Open the file.
+	*
+	* @param Path - The path to the file
+	* @returns true if the file was opened successfully
+	*/
+	UE_API bool Open(const TCHAR* Path);
+
+	UE_API virtual int32 Write(void* Data, uint32 Size) override;
+	UE_API virtual void Close() override;
+
+private:
+	TUniquePtr<IFileHandle> Handle;
+};
+
+} // namespace UE::Trace
+
+#undef UE_API
diff --git a/Engine/Source/Developer/TraceAnalysis/Public/Trace/TraceWriter.h b/Engine/Source/Developer/TraceAnalysis/Public/Trace/TraceWriter.h
new file mode 100644
--- /dev/null
+++ b/Engine/Source/Developer/TraceAnalysis/Public/Trace/TraceWriter.h
@@ -0,0 +1,894 @@
+// Copyright Epic Games, Inc. All Rights Reserved.
+
+#pragma once
+
+#include "CoreTypes.h"
+#include "Containers/AnsiString.h"
+#include "Containers/UnrealString.h"
+#include "Misc/EnumClassFlags.h"
+#include "Templates/Requires.h"
+
+// TraceLog
+#include "Trace/Detail/Protocol.h"
+#include "Trace/Detail/Transport.h"
+
+// TraceAnalysis
+#include "Trace/Analyzer.h"
+
+#include <type_traits>
+
+#define UE_API TRACEANALYSIS_API
+
+namespace UE::Trace
+{
+
+class IOutDataStream;
+class FTraceWriter;
+
+////////////////////////////////////////////////////////////////////////////////////////////////////
+/**
+ * Flags for a trace event.
+ * See UE::Trace::Protocol6::EEventFlags (Trace\Detail\Protocols\Protocol6.h in TraceLog module).
+ * See UE::Trace::IAnalyzer::EEventFlags (Trace\Analyzer.h in TraceAnalysis module).
+ */
+enum class ETraceWriterEventFlags : uint8
+{
+	None            = 0,
+
+	Important       = 1 << 0,
+	MaybeHasAux     = 1 << 1,
+	NoSync          = 1 << 2,
+	Definition      = 1 << 3,
+
+	ImportantNoSync = Important | NoSync,
+};
+ENUM_CLASS_FLAGS(ETraceWriterEventFlags);
+
+////////////////////////////////////////////////////////////////////////////////////////////////////
+/**
+ * Field type.
+ * See UE::Trace::Protocol0::EFieldType (Trace\Detail\Protocols\Protocol0.h in TraceLog module).
+ * { Bool = 0, Int8, Int16, Int32, Int64, Uint8 = 0, Uint16, Uint32, Uint64, ..., Pointer = Uint64, ..., Array }
+ * See also UE::Trace::IAnalyzer::FEventFieldInfo::EType
+ * { None, Integer, Float, AnsiString, WideString, Reference8, Reference16, Reference32, Reference64 }
+ */
+enum class ETraceWriterFieldType : uint8
+{
+	ArrayFlag           = 0x80,
+	ReferenceFlag       = 0x40, // can be combined with Uint8,Uint16,Uint32,Uint64 types
+	DefinitionIdFlag    = 0x20, // can be combined with Uint8,Uint16,Uint32,Uint64 types
+	FlagsMask           = 0xE0,
+	IndexMask           = 0x1F,
+
+	None                = 0,
+
+	// Boolean type
+	Bool                = 1, // 1 byte for a single bool value, variable size (1 bit / value) for a bool array
+
+	// Unsigned integer types
+	Uint8               = 2,
+	Uint16              = 3,
+	Uint32              = 4,
+	//Uint32_7bit         = 5, // variable size, 1 to 5 bytes
+	Uint64              = 6,
+	//Uint64_7bit         = 7, // variable size, 1 to 10 bytes
+
+	// Signed integer types
+	Int8                = 8,
+	Int16               = 9,
+	Int32               = 10,
+	//Int32_7bit          = 11, // variable size, 1 to 5 bytes
+	Int64               = 12,
+	//Int64_7bit          = 13, // variable size, 1 to 10 bytes
+
+	// Floating-point number types
+	Float32             = 14, // float
+	Float64             = 15, // double
+
+	// String types
+	AnsiString          = 16,
+	WideString          = 17,
+	//Utf8String          = 18,
+
+	Pointer             = 19, // void* (64 bits)
+
+	Count               = 20,
+
+	// Reference composed types
+	Reference8          = ReferenceFlag | Uint8,
+	Reference16         = ReferenceFlag | Uint16,
+	Reference32         = ReferenceFlag | Uint32,
+	Reference64         = ReferenceFlag | Uint64,
+
+	// DefinitionId composed types
+	DefinitionId8       = DefinitionIdFlag | Uint8,
+	DefinitionId16      = DefinitionIdFlag | Uint16,
+	DefinitionId32      = DefinitionIdFlag | Uint32,
+	DefinitionId64      = DefinitionIdFlag | Uint64,
+};
+ENUM_CLASS_FLAGS(ETraceWriterFieldType);
+
+////////////////////////////////////////////////////////////////////////////////////////////////////
+
+struct FTraceWriterFieldTypeResolver
+{
+	template <typename NativeType>
+	static constexpr ETraceWriterFieldType GetType()
+	{
+		using EType = ETraceWriterFieldType;
+
+		     if constexpr (std::is_same_v<std::remove_const_t<NativeType>, bool  >) return EType::Bool;
+		else if constexpr (std::is_same_v<std::remove_const_t<NativeType>, uint8 >) return EType::Uint8;
+		else if constexpr (std::is_same_v<std::remove_const_t<NativeType>, uint16>) return EType::Uint16;
+		else if constexpr (std::is_same_v<std::remove_const_t<NativeType>, uint32>) return EType::Uint32;
+		else if constexpr (std::is_same_v<std::remove_const_t<NativeType>, uint64>) return EType::Uint64;
+		else if constexpr (std::is_same_v<std::remove_const_t<NativeType>, int8  >) return EType::Int8;
+		else if constexpr (std::is_same_v<std::remove_const_t<NativeType>, int16 >) return EType::Int16;
+		else if constexpr (std::is_same_v<std::remove_const_t<NativeType>, int32 >) return EType::Int32;
+		else if constexpr (std::is_same_v<std::remove_const_t<NativeType>, int64 >) return EType::Int64;
+		else if constexpr (std::is_same_v<std::remove_const_t<NativeType>, float >) return EType::Float32;
+		else if constexpr (std::is_same_v<std::remove_const_t<NativeType>, double>) return EType::Float64;
+
+		else if constexpr (std::is_same_v<std::remove_const_t<NativeType>, bool[]  >) return EType::ArrayFlag | EType::Bool;
+		else if constexpr (std::is_same_v<std::remove_const_t<NativeType>, uint8[] >) return EType::ArrayFlag | EType::Uint8;
+		else if constexpr (std::is_same_v<std::remove_const_t<NativeType>, uint16[]>) return EType::ArrayFlag | EType::Uint16;
+		else if constexpr (std::is_same_v<std::remove_const_t<NativeType>, uint32[]>) return EType::ArrayFlag | EType::Uint32;
+		else if constexpr (std::is_same_v<std::remove_const_t<NativeType>, uint64[]>) return EType::ArrayFlag | EType::Uint64;
+		else if constexpr (std::is_same_v<std::remove_const_t<NativeType>, int8[]  >) return EType::ArrayFlag | EType::Int8;
+		else if constexpr (std::is_same_v<std::remove_const_t<NativeType>, int16[] >) return EType::ArrayFlag | EType::Int16;
+		else if constexpr (std::is_same_v<std::remove_const_t<NativeType>, int32[] >) return EType::ArrayFlag | EType::Int32;
+		else if constexpr (std::is_same_v<std::remove_const_t<NativeType>, int64[] >) return EType::ArrayFlag | EType::Int64;
+		else if constexpr (std::is_same_v<std::remove_const_t<NativeType>, float[] >) return EType::ArrayFlag | EType::Float32;
+		else if constexpr (std::is_same_v<std::remove_const_t<NativeType>, double[]>) return EType::ArrayFlag | EType::Float64;
+
+		else if constexpr (std::is_same_v<NativeType, TConstArrayView<bool>  >) return EType::ArrayFlag | EType::Bool;
+		else if constexpr (std::is_same_v<NativeType, TConstArrayView<uint8> >) return EType::ArrayFlag | EType::Uint8;
+		else if constexpr (std::is_same_v<NativeType, TConstArrayView<uint16>>) return EType::ArrayFlag | EType::Uint16;
+		else if constexpr (std::is_same_v<NativeType, TConstArrayView<uint32>>) return EType::ArrayFlag | EType::Uint32;
+		else if constexpr (std::is_same_v<NativeType, TConstArrayView<uint64>>) return EType::ArrayFlag | EType::Uint64;
+		else if constexpr (std::is_same_v<NativeType, TConstArrayView<int8>  >) return EType::ArrayFlag | EType::Int8;
+		else if constexpr (std::is_same_v<NativeType, TConstArrayView<int16> >) return EType::ArrayFlag | EType::Int16;
+		else if constexpr (std::is_same_v<NativeType, TConstArrayView<int32> >) return EType::ArrayFlag | EType::Int32;
+		else if constexpr (std::is_same_v<NativeType, TConstArrayView<int64> >) return EType::ArrayFlag | EType::Int64;
+		else if constexpr (std::is_same_v<NativeType, TConstArrayView<float> >) return EType::ArrayFlag | EType::Float32;
+		else if constexpr (std::is_same_v<NativeType, TConstArrayView<double>>) return EType::ArrayFlag | EType::Float64;
+
+		else if constexpr (std::is_same_v<NativeType, TArrayView<bool>  >) return EType::ArrayFlag | EType::Bool;
+		else if constexpr (std::is_same_v<NativeType, TArrayView<uint8> >) return EType::ArrayFlag | EType::Uint8;
+		else if constexpr (std::is_same_v<NativeType, TArrayView<uint16>>) return EType::ArrayFlag | EType::Uint16;
+		else if constexpr (std::is_same_v<NativeType, TArrayView<uint32>>) return EType::ArrayFlag | EType::Uint32;
+		else if constexpr (std::is_same_v<NativeType, TArrayView<uint64>>) return EType::ArrayFlag | EType::Uint64;
+		else if constexpr (std::is_same_v<NativeType, TArrayView<int8>  >) return EType::ArrayFlag | EType::Int8;
+		else if constexpr (std::is_same_v<NativeType, TArrayView<int16> >) return EType::ArrayFlag | EType::Int16;
+		else if constexpr (std::is_same_v<NativeType, TArrayView<int32> >) return EType::ArrayFlag | EType::Int32;
+		else if constexpr (std::is_same_v<NativeType, TArrayView<int64> >) return EType::ArrayFlag | EType::Int64;
+		else if constexpr (std::is_same_v<NativeType, TArrayView<float> >) return EType::ArrayFlag | EType::Float32;
+		else if constexpr (std::is_same_v<NativeType, TArrayView<double>>) return EType::ArrayFlag | EType::Float64;
+
+		else if constexpr (std::is_same_v<NativeType, void*                       >) return EType::Pointer;
+		else if constexpr (std::is_same_v<NativeType, const void*                 >) return EType::Pointer;
+		else if constexpr (std::is_same_v<NativeType, void* []                    >) return EType::ArrayFlag | EType::Pointer;
+		else if constexpr (std::is_same_v<NativeType, const void* []              >) return EType::ArrayFlag | EType::Pointer;
+		else if constexpr (std::is_same_v<NativeType, TConstArrayView<void*>      >) return EType::ArrayFlag | EType::Pointer;
+		else if constexpr (std::is_same_v<NativeType, TConstArrayView<const void*>>) return EType::ArrayFlag | EType::Pointer;
+		else if constexpr (std::is_same_v<NativeType, TArrayView<void*>           >) return EType::ArrayFlag | EType::Pointer;
+		else if constexpr (std::is_same_v<NativeType, TArrayView<const void*>     >) return EType::ArrayFlag | EType::Pointer;
+
+		else if constexpr (std::is_same_v<std::remove_const_t<NativeType>, AnsiString>) return EType::AnsiString;
+		else if constexpr (std::is_same_v<std::remove_const_t<NativeType>, WideString>) return EType::WideString;
+
+		// No default return value (to generate a compile error for unsupported native types)!
+	}
+};
+
+////////////////////////////////////////////////////////////////////////////////////////////////////
+
+class FTraceWriterEventFieldInfo final
+{
+	friend class FTraceWriter;
+	friend class FTraceWriterEventDeclarationBuilder;
+
+public:
+	FTraceWriterEventFieldInfo() = delete;
+	FTraceWriterEventFieldInfo(FAnsiStringView InFieldName, ETraceWriterFieldType InFieldType)
+		: Name(FAnsiString(InFieldName))
+		, Type(InFieldType)
+	{
+	}
+
+	[[nodiscard]] const FAnsiString& GetName() const { return Name; }
+	[[nodiscard]] ETraceWriterFieldType GetType() const { return Type; }
+	[[nodiscard]] uint16 GetSize() const { return Size; }
+	[[nodiscard]] uint16 GetRefUid() const { return RefUid; }
+
+	[[nodiscard]] inline uint8 GetProtocolFieldFamily() const;
+	[[nodiscard]] inline uint8 GetProtocolFieldType() const;
+
+private:
+	FAnsiString Name;
+	ETraceWriterFieldType Type = ETraceWriterFieldType::None;
+	uint16 Size = 0;
+	uint16 RefUid = 0;
+};
+
+////////////////////////////////////////////////////////////////////////////////////////////////////
+
+class FTraceWriterEventInfo final
+{
+	friend class FTraceWriter;
+	friend class FTraceWriterEventDeclarationBuilder;
+
+private:
+	FTraceWriterEventInfo() = delete;
+
+	// Constructor only accessible from the FTraceWriter class
+	FTraceWriterEventInfo(FAnsiStringView InLoggerName, FAnsiStringView InEventName, ETraceWriterEventFlags InEventFlags)
+		: LoggerName(FAnsiString(InLoggerName))
+		, Name(FAnsiString(InEventName))
+		, Flags(InEventFlags)
+	{
+	}
+
+public:
+	[[nodiscard]] uint32 GetId() const { return Id; }
+	[[nodiscard]] const FAnsiString& GetLoggerName() const { return LoggerName; }
+	[[nodiscard]] const FAnsiString& GetName() const { return Name; }
+	[[nodiscard]] ETraceWriterEventFlags GetFlags() const { return Flags; }
+	[[nodiscard]] uint32 GetNumFields() const { return static_cast<uint32>(Fields.Num()); }
+	[[nodiscard]] const FTraceWriterEventFieldInfo& GetField(uint32 InFieldIndex) const { return Fields[InFieldIndex]; }
+	[[nodiscard]] const TArray<FTraceWriterEventFieldInfo>& GetFields() const { return Fields; }
+	[[nodiscard]] UE_API uint32 GetFieldIndex(FAnsiStringView InFieldName) const;
+
+private:
+	uint32 Id = 0;
+	FAnsiString LoggerName;
+	FAnsiString Name;
+	ETraceWriterEventFlags Flags = ETraceWriterEventFlags::None;
+	TArray<FTraceWriterEventFieldInfo> Fields;
+};
+
+////////////////////////////////////////////////////////////////////////////////////////////////////
+
+class FTraceWriterEventDeclarationBuilder final
+{
+	friend class FTraceWriter;
+
+private:
+	FTraceWriterEventDeclarationBuilder() = delete;
+
+	// Constructor only accessible from the FTraceWriter class
+	FTraceWriterEventDeclarationBuilder(FTraceWriter& InWriter, FTraceWriterEventInfo& InEventInfo)
+		: Writer(InWriter)
+		, EventInfo(InEventInfo)
+	{
+	}
+
+	~FTraceWriterEventDeclarationBuilder();
+
+public:
+	[[nodiscard]] FTraceWriterEventInfo& GetEventInfo() const
+	{
+		return EventInfo;
+	}
+
+	/**
+	 * Adds a new regular field to the current event declaration.
+	 *
+	 * @param InFieldName - The name of the field.
+	 * @param InFieldType - The type of the field.
+	 * @returns - The event declaration builder object.
+	 */
+	UE_API FTraceWriterEventDeclarationBuilder& Field(FAnsiStringView InFieldName, ETraceWriterFieldType InFieldType);
+
+	/**
+	 * Adds a new Reference field to the current event declaration.
+	 *
+	 * @param InFieldName - The name of the field.
+	 * @param InFieldType - The type of the field. Can only be Reference8, Reference16, Reference32 or Reference64.
+	 * @param InRefUid - The reference unique id.
+	 * @returns - The event declaration builder object.
+	 */
+	UE_API FTraceWriterEventDeclarationBuilder& ReferenceField(FAnsiStringView InFieldName, ETraceWriterFieldType InFieldType, uint16 InRefUid);
+
+	/**
+	 * Adds a new DefinitionId field to the current event declaration.
+	 * @param InFieldType - The type of the field. Can only be DefinitionId8, DefinitionId16, DefinitionId32 or DefinitionI64.
+	 * @returns - The event declaration builder object.
+	 */
+	UE_API FTraceWriterEventDeclarationBuilder& DefinitionIdField(ETraceWriterFieldType InFieldType);
+
+	/**
+	 * Signals the event declaration is complete.
+	 *
+	 * @returns the event id of the newly declared event
+	 */
+	[[nodiscard]] UE_API uint32 End();
+
+private:
+	FTraceWriter& Writer;
+	FTraceWriterEventInfo& EventInfo;
+	bool bIsCompleted = false;
+};
+
+////////////////////////////////////////////////////////////////////////////////////////////////////
+
+class FTraceWriterFieldData
+{
+public:
+	FTraceWriterFieldData() = default;
+	~FTraceWriterFieldData();
+
+	const void* GetData() const
+	{
+		return Size <= MaxInlineDataSize ? &Data : Buffer;
+	}
+	uint32 GetSize() const
+	{
+		return Size;
+	}
+	UE_API void SetData(const void* InData, uint32 InDataSize);
+
+private:
+	static constexpr uint32 MaxInlineDataSize = 12; // Data + Data2
+	union
+	{
+		void* Buffer;
+		uint64 Data = 0;
+	};
+	uint32 Data2 = 0;
+	uint32 Size = 0;
+};
+static_assert(sizeof(FTraceWriterFieldData) == 16, "");
+
+////////////////////////////////////////////////////////////////////////////////////////////////////
+
+class FTraceWriterEventBuilder final
+{
+	friend class FTraceWriter;
+
+private:
+	FTraceWriterEventBuilder(FTraceWriter& InWriter, const FTraceWriterEventInfo& InEventInfo)
+		: Writer(InWriter)
+		, EventInfo(InEventInfo)
+	{
+		if (InEventInfo.GetNumFields() > 0)
+		{
+			Fields.AddDefaulted(InEventInfo.GetNumFields());
+		}
+	}
+
+	~FTraceWriterEventBuilder();
+
+public:
+	[[nodiscard]] const FTraceWriterEventInfo& GetEventInfo() const
+	{
+		return EventInfo;
+	}
+
+	//////////////////////////////////////////////////
+	// Array types (T[])
+	// Value is passed as Data + DataSize [bytes], as TConstArrayView<T> or as TArrayView<T>.
+
+	UE_API FTraceWriterEventBuilder& Field(uint32 Index, const void* Data, uint32 DataSize);
+
+	FTraceWriterEventBuilder& Field(FAnsiStringView Name, const void* Data, uint32 DataSize)
+	{
+		return Field(EventInfo.GetFieldIndex(Name), Data, DataSize);
+	}
+
+	template <typename T>
+	FTraceWriterEventBuilder& Field(uint32 Index, TConstArrayView<T> ArrayView)
+	{
+		return Field(Index, ArrayView.GetData(), ArrayView.Num() * sizeof(T));
+	}
+
+	template <typename T>
+	FTraceWriterEventBuilder& Field(uint32 Index, TArrayView<T> ArrayView)
+	{
+		return Field(Index, ArrayView.GetData(), ArrayView.Num() * sizeof(T));
+	}
+
+	//////////////////////////////////////////////////
+	// Integral and floating point types
+	// (bool, uint8, uint16,..., float, double)
+
+	template <
+		typename T
+		UE_REQUIRES(
+			std::is_integral_v<T> || std::is_floating_point_v<T>
+		)
+	>
+	FTraceWriterEventBuilder& Field(uint32 Index, T Value)
+	{
+		constexpr ETraceWriterFieldType FieldType = FTraceWriterFieldTypeResolver::GetType<T>();
+		if (CheckValidPodField(Index, FieldType, sizeof(T)))
+		{
+			Fields[Index].SetData(&Value, sizeof(T));
+		}
+		return *this;
+	}
+
+	//////////////////////////////////////////////////
+	// Pointer type (void*)
+
+	FTraceWriterEventBuilder& Field(uint32 Index, const void* Value)
+	{
+		if (CheckValidPodField(Index, ETraceWriterFieldType::Pointer, sizeof(void*)))
+		{
+			Fields[Index].SetData(&Value, sizeof(void*));
+		}
+		return *this;
+	}
+
+	FTraceWriterEventBuilder& Field(FAnsiStringView Name, const void* Value)
+	{
+		return Field(EventInfo.GetFieldIndex(Name), Value);
+	}
+
+	FTraceWriterEventBuilder& Field(uint32 Index, void* Value)
+	{
+		if (CheckValidPodField(Index, ETraceWriterFieldType::Pointer, sizeof(void*)))
+		{
+			Fields[Index].SetData(&Value, sizeof(void*));
+		}
+		return *this;
+	}
+
+	FTraceWriterEventBuilder& Field(FAnsiStringView Name, void* Value)
+	{
+		return Field(EventInfo.GetFieldIndex(Name), Value);
+	}
+
+	//////////////////////////////////////////////////
+	// AnsiString
+
+	/** Construct field from an FAnsiStringView. */
+	template <typename T = FAnsiStringView>
+	FTraceWriterEventBuilder& Field(uint32 Index, FAnsiStringView Value UE_LIFETIMEBOUND)
+	{
+		if (CheckValidExactFieldType(Index, ETraceWriterFieldType::AnsiString))
+		{
+			SetField(Index, (const void*)Value.GetData(), (uint32)Value.NumBytes());
+		}
+		return *this;
+	}
+
+	/** Construct field from the null-terminated string pointed to by InData. */
+	template <typename T = const ANSICHAR*>
+	FTraceWriterEventBuilder& Field(uint32 Index, const ANSICHAR* InData UE_LIFETIMEBOUND)
+	{
+		return Field(Index, FAnsiStringView(InData));
+	}
+
+	/** Construct field from an FAnsiString. */
+	template <typename T = const FAnsiString&>
+	FTraceWriterEventBuilder& Field(uint32 Index, const FAnsiString& Value UE_LIFETIMEBOUND)
+	{
+		return Field(Index, FAnsiStringView(Value));
+	}
+
+	//////////////////////////////////////////////////
+	// WideString
+
+	/** Construct field from an FStringView. */
+	template <typename T = FStringView>
+	FTraceWriterEventBuilder& Field(uint32 Index, const FStringView Value UE_LIFETIMEBOUND)
+	{
+		if (CheckValidExactFieldType(Index, ETraceWriterFieldType::WideString))
+		{
+			SetField(Index, Value.GetData(), (uint32)Value.NumBytes());
+		}
+		return *this;
+	}
+
+	/** Construct field from the null-terminated string pointed to by InData. */
+	template <typename T = const TCHAR*>
+	FTraceWriterEventBuilder& Field(uint32 Index, const TCHAR* Value UE_LIFETIMEBOUND)
+	{
+		return Field(Index, FStringView(Value));
+	}
+
+	/** Construct field from an FString. */
+	template <typename T = const FString&>
+	FTraceWriterEventBuilder& Field(uint32 Index, const FString& Value UE_LIFETIMEBOUND)
+	{
+		return Field(Index, FStringView(Value));
+	}
+
+	//////////////////////////////////////////////////
+	// Identify the field by name.
+
+	template <typename T>
+	FTraceWriterEventBuilder& Field(FAnsiStringView Name, T Value)
+	{
+		return Field(EventInfo.GetFieldIndex(Name), Value);
+	}
+
+	//////////////////////////////////////////////////
+
+	UE_API void End();
+
+private:
+	void SetErrorInvalidFieldIndex(uint32 Index);
+	void SetErrorInvalidFieldType(uint32 Index);
+	void SetErrorInvalidFieldData(uint32 Index);
+
+	bool CheckValidPodField(uint32 Index, ETraceWriterFieldType FieldType, uint32 Size);
+	bool CheckValidExactFieldType(uint32 Index, ETraceWriterFieldType FieldType);
+	void SetField(uint32 Index, const void* Data, uint32 DataSize);
+
+private:
+	FTraceWriter& Writer;
+	const FTraceWriterEventInfo& EventInfo;
+	TArray<FTraceWriterFieldData> Fields;
+	bool bIsCompleted = false;
+};
+
+////////////////////////////////////////////////////////////////////////////////////////////////////
+
+class FTraceWriter
+{
+	friend class FTraceWriterEventInfo;
+	friend class FTraceWriterEventDeclarationBuilder;
+	friend class FTraceWriterEventBuilder;
+
+private:
+	struct FTraceGuid
+	{
+		uint32 Bits[4];
+	};
+	struct FThreadInfo
+	{
+		uint32 Id;
+		FAnsiString Name;
+		uint32 SystemId = 0;
+		int32 SortHint = 0;
+	};
+
+public:
+	UE_API FTraceWriter(IOutDataStream& InStream);
+	UE_API ~FTraceWriter();
+
+	FTraceWriter(const FTraceWriter&) = delete;
+	FTraceWriter& operator=(const FTraceWriter&) = delete;
+
+	/**
+	 * Returns the version of the transport protocol used by this trace writer.
+	 */
+	UE_API UE::Trace::ETransport GetTransportProtocolVersion() const;
+
+	/**
+	 * Returns the version of the trace protocol used by this trace writer.
+	 * Currently this is Protocol 7.
+	 *
+	 * \details
+	 * See UE::Trace::Protocol7::EProtocol::Id (Trace/Detail/Protocols/Protocol7.h).
+	 */
+	UE_API uint8 GetTraceProtocolVersion() const;
+
+	/**
+	 * Begins the trace stream.
+	 *
+	 * \details
+	 * Declares: $Trace.* events
+	 * Writes: $Trace.NewEvent, $Trace.NewTrace
+	 */
+	UE_API void Begin();
+
+	/**
+	 * Ends the trace stream.
+	 *
+	 * \details
+	 * Calls Close() for the stream.
+	 */
+	UE_API void End();
+
+	/**
+	 * Writes a $Trace.ThreadGroupBegin event.
+	 * Changes the current thread to be the "Importants" thread.
+	 *
+	 * \details
+	 * Writes: $Trace.ThreadGroupBegin
+	 *
+	 * @param GroupName - The group name
+	 */
+	UE_API void WriteThreadGroupBeginEvent(FAnsiStringView GroupName);
+
+	/**
+	 * Writes a $Trace.ThreadGroupEnd event.
+	 * Changes the current thread to be the "Importants" thread.
+	 *
+	 * \details
+	 * Writes: $Trace.ThreadGroupEnd
+	 */
+	UE_API void WriteThreadGroupEndEvent();
+
+	/**
+	 * Declares a new trace thread.
+	 * Writes the $Trace.ThreadInfo event.
+	 * Changes the current thread to be the "Importants" thread.
+	 *
+	 * \details
+	 * Writes: $Trace.ThreadInfo
+	 *
+	 * @param ThreadName - The thread name
+	 * @param SystemId - The system thread id
+	 * @param SortHint - A hint for sorting priority
+	 */
+	UE_API uint32 DeclareThread(FAnsiStringView ThreadName, uint32 SystemId = 0, int32 SortHint = 0);
+
+	/**
+	 * Sets the current thread id.
+	 *
+	 * @param ThreadId - The thread id.
+	 */
+	UE_API void SetCurrentThread(uint32 ThreadId);
+
+	/**
+	 * Sets the current thread to be the "Events" thread.
+	 */
+	void SetCurrentThreadEvents() { SetCurrentThread(static_cast<uint32>(ETransportTid::Events)); }
+
+	/**
+	 * Sets the current thread to be the "Importants" thread.
+	 */
+	void SetCurrentThreadImportants() { SetCurrentThread(static_cast<uint32>(ETransportTid::Importants)); }
+
+	/**
+	 * Gets the current thread id.
+	 *
+	 * @returns the current thread id
+	 */
+	uint32 GetCurrentThread() const { return CurrentThreadId; }
+
+	/**
+	 * Writes a $Trace.ThreadTiming event for the current thread.
+	 *
+	 * @param BaseTimestamp - The base timestamp [Cycles64]
+	 */
+	UE_API void WriteThreadTimingEvent(uint64 BaseTimestamp);
+
+	/**
+	 * Declares a new trace event.
+	 * Changes the current thread to be the "Events" thread.
+	 *
+	 * \details
+	 * Declares: the new trace event
+	 * Writes: $Trace.NewEvent
+	 *
+	 * @param LoggerName - The logger name
+	 * @param Name - The event name
+	 * @param Flags - The event flags (Important, NoSync, etc.)
+	 * @returns an event declaration builder object (FTraceWriterEventDeclarationBuilder);
+	 *          The builder object is used to further declare the event fields.
+	 *          User needs to call the builder's End() to complete the event declaration
+	 *          and to get the event id for the newly declared event.
+	 */
+	[[nodiscard]] UE_API FTraceWriterEventDeclarationBuilder& DeclareEvent(FAnsiStringView LoggerName, FAnsiStringView Name, ETraceWriterEventFlags Flags = ETraceWriterEventFlags::None);
+
+	/**
+	 * Gets the info for an event identified by Id.
+	 *
+	 * @param EventId - The event id, one returned by DeclareEvent().[...].End() or similar method.
+	 * @returns a pointer to the event info or nullptr if the EventId is not valid.
+	 */
+	const FTraceWriterEventInfo* GetEventInfo(uint32 EventId) const
+	{
+		return EventId < uint32(EventInfos.Num()) ? EventInfos[EventId] : nullptr;
+	}
+
+	/**
+	 * Declares the Diagnostics.Session2 event.
+	 *
+	 * \details
+	 * Declares: Diagnostics.Session2 event
+	 * Writes: $Trace.NewEvent
+	 */
+	[[nodiscard]] UE_API uint32 DeclareDiagnosticsSession2Event();
+
+	/**
+	 * Writes the specified event, on the current thread.
+	 *
+	 * @param EventId - The event id, one returned by DeclareEvent().[...].End()
+	 * @returns an event builder object (FTraceWriterEventBuilder);
+	 *          The builder object is used to further set values for each event field.
+	 *          User needs to call the builder's End() to complete the event.
+	 */
+	[[nodiscard]] UE_API FTraceWriterEventBuilder& WriteEvent(uint32 EventId);
+
+	//////////////////////////////////////////////////
+	// Event Flags
+
+	[[nodiscard]] static ETraceWriterEventFlags ConvertEventFlags(UE::Trace::IAnalyzer::EEventFlags Flags)
+	{
+		static_assert((uint8)UE::Trace::IAnalyzer::EEventFlags::None == (uint8)ETraceWriterEventFlags::None, "");
+		static_assert((uint8)UE::Trace::IAnalyzer::EEventFlags::Important == (uint8)ETraceWriterEventFlags::Important, "");
+		static_assert((uint8)UE::Trace::IAnalyzer::EEventFlags::MaybeHasAux == (uint8)ETraceWriterEventFlags::MaybeHasAux, "");
+		static_assert((uint8)UE::Trace::IAnalyzer::EEventFlags::Definition == (uint8)ETraceWriterEventFlags::Definition, "");
+		return static_cast<ETraceWriterEventFlags>(Flags);
+	}
+
+	[[nodiscard]] static uint8 GetProtocolEventFlags(ETraceWriterEventFlags Flags)
+	{
+		using namespace UE::Trace::Protocol7;
+		static_assert((uint8)ETraceWriterEventFlags::None == (uint8)EEventFlags::None, "");
+		static_assert((uint8)ETraceWriterEventFlags::Important == (uint8)EEventFlags::Important, "");
+		static_assert((uint8)ETraceWriterEventFlags::MaybeHasAux == (uint8)EEventFlags::MaybeHasAux, "");
+		static_assert((uint8)ETraceWriterEventFlags::Definition == (uint8)EEventFlags::Definition, "");
+		return static_cast<uint8>(Flags);
+	}
+
+	//////////////////////////////////////////////////
+	// Field Types
+
+	UE_API static uint16 GetByteSizeForFieldType(ETraceWriterFieldType Type);
+
+	[[nodiscard]] static bool IsArrayFieldType(ETraceWriterFieldType Type)
+	{
+		return EnumHasAnyFlags(Type, ETraceWriterFieldType::ArrayFlag);
+	}
+	[[nodiscard]] static bool IsSignedFieldType(ETraceWriterFieldType Type)
+	{
+		return (uint32)Type >= (uint32)ETraceWriterFieldType::Int8 && (uint32)Type <= (uint32)ETraceWriterFieldType::Float64;
+	}
+	[[nodiscard]] static bool IsReferenceFieldType(ETraceWriterFieldType Type)
+	{
+		return EnumHasAnyFlags(Type, ETraceWriterFieldType::ReferenceFlag);
+	}
+	[[nodiscard]] static bool IsDefinitionIdFieldType(ETraceWriterFieldType Type)
+	{
+		return EnumHasAnyFlags(Type, ETraceWriterFieldType::DefinitionIdFlag);
+	}
+
+	[[nodiscard]] UE_API static uint8 GetProtocolFieldFamily(ETraceWriterFieldType Type);
+	[[nodiscard]] UE_API static uint8 GetProtocolFieldType(ETraceWriterFieldType Type);
+
+	[[nodiscard]] UE_API static ETraceWriterFieldType ConvertFieldType(const UE::Trace::IAnalyzer::FEventFieldInfo& FieldInfo);
+
+	//////////////////////////////////////////////////
+
+	[[nodiscard]] UE_API uint64 GetTimeFrequency() const;
+	[[nodiscard]] UE_API uint64 GetTimeCycles64() const;
+	[[nodiscard]] UE_API double GetTimeSinceEpoch() const;
+
+	//UE_API void SetCustomStartTime(uint64 InStartTimeFrequency, uint64 InStartTimeCycles64, double InStartTimeSinceEpoch);
+	//UE_API void SetCustomClock(TFunctionRef<uint64> TimeFrequencyGetter, TFunctionRef<uint64> TimeCycles64Getter, TFunctionRef<double> TimeSinceEpochGetter);
+
+	[[nodiscard]] uint64 GetStartTimeFrequency() const { return CycleFrequency; }
+	[[nodiscard]] uint64 GetStartTimeCycles64() const { return StartCycle; }
+	[[nodiscard]] double GetStartTimeSinceEpoch() const { return StartDateTime; }
+
+	//////////////////////////////////////////////////
+
+	/**
+	 * Returns the last error message from the stack. Returns an empty string if there was no error.
+	 * Continuing to call GetLastError() will extract the previous emitted error until the stack is empty.
+	 */
+	[[nodiscard]] UE_API const FString& GetLastError() const;
+
+	/**
+	 * Clears the stack of error messages, if any.
+	 */
+	void ClearErrors() { Errors.Reset(); }
+
+	/**
+	 * Returns wherever safe mode is enabled or not.
+	 */
+	bool IsSafeModeEnabled() const { return true; }
+
+private:
+	void Reset();
+	static void GenerateGuid(FTraceGuid* OutGuid);
+	void RegisterDefaultEvents();
+	uint32 GetOrRegisterUnknownEventId();
+
+	void WriteEventDeclaration(const FTraceWriterEventDeclarationBuilder& Builder);
+	void WriteEventDeclaration(const FTraceWriterEventInfo& EventInfo);
+	void WriteEvent(const FTraceWriterEventBuilder& Builder);
+	void WriteEvent(const FTraceWriterEventBuilder& Builder, const FTraceWriterEventInfo& EventInfo);
+	void WriteAuxData(uint32 Index, const void* Data, uint32 Size, bool bIsImportant);
+	void WriteAuxDataSegment(uint32 Index, const void* SegmentData, uint32 SegmentSize, bool bIsImportant);
+
+	uint8* GetBuffer(uint32 RequiredDataSize);
+	void AdvanceBuffer(uint32 DataSize);
+	void BeginEvent();
+	void EndEvent();
+	void WritePacket();
+
+	void SetError(FString&& InError);
+
+private:
+	IOutDataStream& Stream;
+
+	FTraceGuid TraceGuid;
+	FTraceGuid SessionGuid;
+
+	uint64 StartCycle = 0;
+	uint64 CycleFrequency = 0;
+	double StartDateTime = 0.0;
+
+	uint32 NextSerial = 0;
+
+	uint32 UnknownEventId = 0;
+	uint32 NewTraceEventId = 0;
+	uint32 ThreadTimingEventId = 0;
+	uint32 ThreadInfoEventId = 0;
+	uint32 ThreadGroupBeginEventId = 0;
+	uint32 ThreadGroupEndEventId = 0;
+
+	TArray<FThreadInfo> Threads;
+
+	TArray<FTraceWriterEventInfo*> EventInfos;
+	FTraceWriterEventDeclarationBuilder* PendingEventDeclaration = nullptr;
+	FTraceWriterEventBuilder* PendingEvent = nullptr;
+
+	uint32 CurrentThreadId = ETransportTid::Events;
+	uint8* EventBuffer = nullptr;
+	uint32 EventBufferSize = 0;
+	static constexpr uint32 MaxPacketHeaderSize = 8; // >= sizeof(FTidPacketEncoded), >= sizeof(FTidPacket)
+	static constexpr uint32 EncodingOverhead = 64; // LZ4_COMPRESSBOUND
+	static constexpr uint32 MaxPacketBufferSize = 64 << 10;
+	static constexpr uint32 MaxEncodedBufferSize = MaxPacketBufferSize - EncodingOverhead - MaxPacketHeaderSize;
+	static constexpr uint32 MaxDecodedBufferSize = MaxEncodedBufferSize;
+	static constexpr uint32 MinEncodedBufferSize = 384;
+	static constexpr uint32 FlushPacketThreshold = 32 << 10;
+	uint32 EventBufferDataSize = 0;
+	uint32 CompletedEventsDataSize = 0;
+	uint8* PacketBuffer = nullptr;
+
+	mutable TArray<FString> Errors;
+	mutable FString LastError;
+};
+
+////////////////////////////////////////////////////////////////////////////////////////////////////
+
+uint8 FTraceWriterEventFieldInfo::GetProtocolFieldFamily() const
+{
+	return FTraceWriter::GetProtocolFieldFamily(Type);
+}
+
+////////////////////////////////////////////////////////////////////////////////////////////////////
+
+uint8 FTraceWriterEventFieldInfo::GetProtocolFieldType() const
+{
+	return FTraceWriter::GetProtocolFieldType(Type);
+}
+
+////////////////////////////////////////////////////////////////////////////////////////////////////
+
+#define UE_ADHOC_TRACE_BEGIN_DECLARE_EVENT(Writer, LoggerName, EventName, Flags) \
+	Writer.DeclareEvent(ANSITEXTVIEW(LoggerName), ANSITEXTVIEW(EventName), ETraceWriterEventFlags:: Flags)
+
+#define UE_ADHOC_TRACE_DECLARE_FIELD(Type, FieldName) \
+	.Field(ANSITEXTVIEW(FieldName), FTraceWriterFieldTypeResolver::GetType<Type>())
+
+#define UE_ADHOC_TRACE_END_DECLARE_EVENT \
+	.End();
+
+#define __UE_ADHOC_TRACE_NARG(...) __UE_ADHOC_TRACE_NARG_(__VA_ARGS__, __UE_ADHOC_TRACE_RSEQ_N())
+#define __UE_ADHOC_TRACE_NARG_(...) __UE_ADHOC_TRACE_ARG_N(__VA_ARGS__)
+#define __UE_ADHOC_TRACE_ARG_N(_1, _2, _3, _4, _5, _6, _7, _8, N, ...) N
+#define __UE_ADHOC_TRACE_RSEQ_N() 8, 7, 6, 5, 4, 3, 2, 1, 0
+
+#define __UE_ADHOC_TRACE_CONCATENATE(a, b) __UE_ADHOC_TRACE_CONCATENATE2(a, b)
+#define __UE_ADHOC_TRACE_CONCATENATE2(a, b) a##b
+
+#define __UE_ADHOC_TRACE_LOG_ARGS_1(Arg0)											.Field(0, Arg0)
+#define __UE_ADHOC_TRACE_LOG_ARGS_2(Arg0, Arg1)										.Field(0, Arg0).Field(1, Arg1)
+#define __UE_ADHOC_TRACE_LOG_ARGS_3(Arg0, Arg1, Arg2)								.Field(0, Arg0).Field(1, Arg1).Field(2, Arg2)
+#define __UE_ADHOC_TRACE_LOG_ARGS_4(Arg0, Arg1, Arg2, Arg3)							.Field(0, Arg0).Field(1, Arg1).Field(2, Arg2).Field(3, Arg3)
+#define __UE_ADHOC_TRACE_LOG_ARGS_5(Arg0, Arg1, Arg2, Arg3, Arg4)					.Field(0, Arg0).Field(1, Arg1).Field(2, Arg2).Field(3, Arg3).Field(4, Arg4)
+#define __UE_ADHOC_TRACE_LOG_ARGS_6(Arg0, Arg1, Arg2, Arg3, Arg4, Arg5)				.Field(0, Arg0).Field(1, Arg1).Field(2, Arg2).Field(3, Arg3).Field(4, Arg4).Field(5, Arg5)
+#define __UE_ADHOC_TRACE_LOG_ARGS_7(Arg0, Arg1, Arg2, Arg3, Arg4, Arg5, Arg6)		.Field(0, Arg0).Field(1, Arg1).Field(2, Arg2).Field(3, Arg3).Field(4, Arg4).Field(5, Arg5).Field(6, Arg6)
+#define __UE_ADHOC_TRACE_LOG_ARGS_8(Arg0, Arg1, Arg2, Arg3, Arg4, Arg5, Arg6, Arg7)	.Field(0, Arg0).Field(1, Arg1).Field(2, Arg2).Field(3, Arg3).Field(4, Arg4).Field(5, Arg5).Field(6, Arg6).Field(7, Arg7)
+
+#define UE_ADHOC_TRACE_LOG(Writer, Uid, ...) \
+	Writer.WriteEvent(Uid) \
+		__UE_ADHOC_TRACE_CONCATENATE(__UE_ADHOC_TRACE_LOG_ARGS_,__UE_ADHOC_TRACE_NARG(__VA_ARGS__))(__VA_ARGS__) \
+		.End();
+
+////////////////////////////////////////////////////////////////////////////////////////////////////
+
+} // namespace UE::Trace
+
+#undef UE_API
