[TraceAnalyzer] Minor update for debug logs.

#rnx
#Synced-CL 49202848

diff --git a/Engine/Source/Developer/TraceAnalysis/Private/Analysis/Transport/TidPacketTransport.cpp b/Engine/Source/Developer/TraceAnalysis/Private/Analysis/Transport/TidPacketTransport.cpp
--- a/Engine/Source/Developer/TraceAnalysis/Private/Analysis/Transport/TidPacketTransport.cpp
+++ b/Engine/Source/Developer/TraceAnalysis/Private/Analysis/Transport/TidPacketTransport.cpp
@@ -130,6 +130,10 @@ FTidPacketTransport::EReadPacketResult FTidPacketTransport::ReadPacket()
 			{
 				UE_TRACE_ANALYSIS_DEBUG_Appendf("%02X ", Dest[i]);
 			}
+			if (DecodedSize > 32)
+			{
+				UE_TRACE_ANALYSIS_DEBUG_Append("...");
+			}
 			UE_TRACE_ANALYSIS_DEBUG_EndStringBuilder();
 #endif // UE_TRACE_ANALYSIS_DEBUG
 		}
@@ -165,6 +169,10 @@ FTidPacketTransport::EReadPacketResult FTidPacketTransport::ReadPacket()
 		{
 			UE_TRACE_ANALYSIS_DEBUG_Appendf("%02X ", ((uint8*)(PacketBase + 1))[i]);
 		}
+		if (DataSize > 32)
+		{
+			UE_TRACE_ANALYSIS_DEBUG_Append("...");
+		}
 		UE_TRACE_ANALYSIS_DEBUG_EndStringBuilder();
 #endif // UE_TRACE_ANALYSIS_DEBUG
 	}
