[Insights] TraceAnalysis: Fixed the array reader's element access for float and double arrays.

#Synced-CL 49273099

diff --git a/Engine/Source/Developer/TraceAnalysis/Private/Analysis/Engine.cpp b/Engine/Source/Developer/TraceAnalysis/Private/Analysis/Engine.cpp
--- a/Engine/Source/Developer/TraceAnalysis/Private/Analysis/Engine.cpp
+++ b/Engine/Source/Developer/TraceAnalysis/Private/Analysis/Engine.cpp
@@ -786,9 +786,9 @@ const void* IAnalyzer::FArrayReader::GetImpl(uint32 Index, int8& SizeAndType) co
 	{
 		return nullptr;
 	}
-
-	SizeAndType = (SizeAndType < 0) ? static_cast<int8>(-SizeAndType) : SizeAndType;
-	return Inner->Data + (Index * SizeAndType);
+	return (SizeAndType < 0) ?
+		Inner->Data + (Index * (-SizeAndType)) :
+		Inner->Data + (Index * SizeAndType);
 }
 
 
