[Insights]
* TraceWriter: Added support for scoped events.
* TraceWriter: Fixed element type for arrays.

#Synced-CL 49273099

diff --git a/Engine/Source/Developer/TraceAnalysis/Private/TraceWriter.cpp b/Engine/Source/Developer/TraceAnalysis/Private/TraceWriter.cpp
--- a/Engine/Source/Developer/TraceAnalysis/Private/TraceWriter.cpp
+++ b/Engine/Source/Developer/TraceAnalysis/Private/TraceWriter.cpp
@@ -1231,6 +1231,89 @@ void FTraceWriter::WriteAuxDataSegment(uint32 Index, const void* SegmentData, ui
 
 ////////////////////////////////////////////////////////////////////////////////////////////////////
 
+void FTraceWriter::WriteEnterScopeEvent()
+{
+	BeginEvent();
+	using namespace Private::WriterProtocol;
+	WriteScopeEventPrivate(uint8(EKnownEventUids::EnterScope << EKnownEventUids::_UidShift));
+	EndEvent();
+}
+
+////////////////////////////////////////////////////////////////////////////////////////////////////
+
+void FTraceWriter::WriteLeaveScopeEvent()
+{
+	BeginEvent();
+	using namespace Private::WriterProtocol;
+	WriteScopeEventPrivate(uint8(EKnownEventUids::LeaveScope << EKnownEventUids::_UidShift));
+	EndEvent();
+}
+
+////////////////////////////////////////////////////////////////////////////////////////////////////
+
+void FTraceWriter::WriteStampedEnterScopeEvent()
+{
+	BeginEvent();
+	using namespace Private::WriterProtocol;
+	constexpr uint8 Uid = uint8(EKnownEventUids::EnterScope_TB << EKnownEventUids::_UidShift);
+	WriteStampedScopeEventPrivate(Uid);
+	EndEvent();
+}
+
+////////////////////////////////////////////////////////////////////////////////////////////////////
+
+void FTraceWriter::WriteStampedLeaveScopeEvent()
+{
+	BeginEvent();
+	using namespace Private::WriterProtocol;
+	constexpr uint8 Uid = uint8(EKnownEventUids::LeaveScope_TB << EKnownEventUids::_UidShift);
+	WriteStampedScopeEventPrivate(Uid);
+	EndEvent();
+}
+
+////////////////////////////////////////////////////////////////////////////////////////////////////
+
+void FTraceWriter::WriteScopeEventPrivate(uint8 Uid)
+{
+	uint8* Buffer = GetBuffer(1);
+	if (!Buffer)
+	{
+		SetError(TEXT("Failed to write *Scope event!"));
+		return;
+	}
+
+	*Buffer = Uid;
+
+	AdvanceBuffer(1);
+}
+
+////////////////////////////////////////////////////////////////////////////////////////////////////
+
+void FTraceWriter::WriteStampedScopeEventPrivate(uint8 Uid)
+{
+	constexpr uint32 Size = sizeof(uint64); // 1 byte (Uid) + 7 bytes (timestamp)
+
+	uint8* Buffer = GetBuffer(Size);
+	if (!Buffer)
+	{
+		SetError(TEXT("Failed to write *Scope_TB event!"));
+		return;
+	}
+
+	const uint64 RelativeTimestamp = GetTimeCycles64() - GetStartTimeCycles64();
+	if ((RelativeTimestamp >> 56) != 0)
+	{
+		SetError(TEXT("Relative timestamp is too large!"));
+	}
+
+	const uint64 Pack = (RelativeTimestamp << 8) + Uid;
+	FMemory::Memcpy(Buffer, &Pack, sizeof(uint64));
+
+	AdvanceBuffer(Size);
+}
+
+////////////////////////////////////////////////////////////////////////////////////////////////////
+
 uint16 FTraceWriter::GetByteSizeForFieldType(ETraceWriterFieldType Type)
 {
 	if (EnumHasAnyFlags(Type, ETraceWriterFieldType::ArrayFlag))
@@ -1300,7 +1383,7 @@ uint8 FTraceWriter::GetProtocolFieldType(ETraceWriterFieldType Type)
 
 	if (EnumHasAnyFlags(Type, ETraceWriterFieldType::ArrayFlag))
 	{
-		return uint8(EFieldType::Array);
+		return uint8(EFieldType::Array) | GetProtocolFieldType(Type & ~ETraceWriterFieldType::ArrayFlag);
 	}
 
 	if (EnumHasAnyFlags(Type, ETraceWriterFieldType::ReferenceFlag))
diff --git a/Engine/Source/Developer/TraceAnalysis/Public/Trace/TraceWriter.h b/Engine/Source/Developer/TraceAnalysis/Public/Trace/TraceWriter.h
--- a/Engine/Source/Developer/TraceAnalysis/Public/Trace/TraceWriter.h
+++ b/Engine/Source/Developer/TraceAnalysis/Public/Trace/TraceWriter.h
@@ -69,17 +69,17 @@ enum class ETraceWriterFieldType : uint8
 	Uint8               = 2,
 	Uint16              = 3,
 	Uint32              = 4,
-	//Uint32_7bit         = 5, // variable size, 1 to 5 bytes
+	//VarUInt32           = 5, // variable size, 1 to 5 bytes
 	Uint64              = 6,
-	//Uint64_7bit         = 7, // variable size, 1 to 10 bytes
+	//VarUInt64           = 7, // variable size, 1 to 9 bytes
 
 	// Signed integer types
 	Int8                = 8,
 	Int16               = 9,
 	Int32               = 10,
-	//Int32_7bit          = 11, // variable size, 1 to 5 bytes
+	//VarInt32            = 11, // variable size, 1 to 5 bytes
 	Int64               = 12,
-	//Int64_7bit          = 13, // variable size, 1 to 10 bytes
+	//VarInt64            = 13, // variable size, 1 to 9 bytes
 
 	// Floating-point number types
 	Float32             = 14, // float
@@ -204,6 +204,8 @@ public:
 	[[nodiscard]] inline uint8 GetProtocolFieldFamily() const;
 	[[nodiscard]] inline uint8 GetProtocolFieldType() const;
 
+	bool IsArray() const { return EnumHasAnyFlags(Type, ETraceWriterFieldType::ArrayFlag); }
+
 private:
 	FAnsiString Name;
 	ETraceWriterFieldType Type = ETraceWriterFieldType::None;
@@ -239,6 +241,12 @@ public:
 	[[nodiscard]] const TArray<FTraceWriterEventFieldInfo>& GetFields() const { return Fields; }
 	[[nodiscard]] UE_API uint32 GetFieldIndex(FAnsiStringView InFieldName) const;
 
+	bool IsNoSync() const { return EnumHasAnyFlags(Flags, ETraceWriterEventFlags::NoSync); }
+	bool IsSync() const { return !EnumHasAnyFlags(Flags, ETraceWriterEventFlags::NoSync); }
+	bool IsImportant() const { return EnumHasAnyFlags(Flags, ETraceWriterEventFlags::Important); }
+	bool MaybeHasAux() const { return EnumHasAnyFlags(Flags, ETraceWriterEventFlags::MaybeHasAux); }
+	bool IsDefinition() const { return EnumHasAnyFlags(Flags, ETraceWriterEventFlags::Definition); }
+
 private:
 	uint32 Id = 0;
 	FAnsiString LoggerName;
@@ -650,7 +658,7 @@ public:
 
 	/**
 	 * Declares a new trace event.
-	 * Changes the current thread to be the "Events" thread.
+	 * The End() call for the returned builder will change the current thread to be the "Events" thread.
 	 *
 	 * \details
 	 * Declares: the new trace event
@@ -696,6 +704,15 @@ public:
 	 */
 	[[nodiscard]] UE_API FTraceWriterEventBuilder& WriteEvent(uint32 EventId);
 
+	//////////////////////////////////////////////////
+	// Scoped Events
+
+	UE_API void WriteEnterScopeEvent();
+	UE_API void WriteLeaveScopeEvent();
+
+	UE_API void WriteStampedEnterScopeEvent();
+	UE_API void WriteStampedLeaveScopeEvent();
+
 	//////////////////////////////////////////////////
 	// Event Flags
 
@@ -789,6 +806,9 @@ private:
 	void WriteAuxData(uint32 Index, const void* Data, uint32 Size, bool bIsImportant);
 	void WriteAuxDataSegment(uint32 Index, const void* SegmentData, uint32 SegmentSize, bool bIsImportant);
 
+	void WriteScopeEventPrivate(uint8 Uid);
+	void WriteStampedScopeEventPrivate(uint8 Uid);
+
 	uint8* GetBuffer(uint32 RequiredDataSize);
 	void AdvanceBuffer(uint32 DataSize);
 	void BeginEvent();
