[TraceAnalysis]
* TraceWriter: Added support for custom clock: SetCustomClock(), GetTimeFrequency(), GetStartTime(), GetStartTimeSinceEpoch(), GetTime(), GetTimeSinceEpoch().
* TraceWriter: Added FindEvent().
* TraceWriter: Added options for a custom Begin() to allow separate calls for DeclareDefaultEvents() and WriteNewTraceEvent().
* TraceWriter: Added explicit timestamp to WriteStamped[Enter/Leave]ScopeEvent() methods.
* TraceWriter: Renamed DeclareThread() to RegisterThread().
* TraceWriter: Added support for custom thread registration: RegisterCustomThread().
* TraceWriter: Added support for writing thread info event: WriteThreadInfo().
* TraceWriter: Added support for getting info for a registered thread: FTraceWriterThreadInfo, GetThreadInfo().
* TraceWriter: Fixed ConvertFieldType() for array types.

#Synced-CL 49645631

diff --git a/Engine/Source/Developer/TraceAnalysis/Private/TraceWriter.cpp b/Engine/Source/Developer/TraceAnalysis/Private/TraceWriter.cpp
--- a/Engine/Source/Developer/TraceAnalysis/Private/TraceWriter.cpp
+++ b/Engine/Source/Developer/TraceAnalysis/Private/TraceWriter.cpp
@@ -353,6 +353,8 @@ FTraceWriter::FTraceWriter(IOutDataStream& InStream)
 	GenerateGuid(&SessionGuid);
 
 	EventInfos.SetNumZeroed(uint32(Private::WriterProtocol::EKnownEventUids::User), EAllowShrinking::No);
+
+	InitDefaultClock();
 }
 
 ////////////////////////////////////////////////////////////////////////////////////////////////////
@@ -384,9 +386,17 @@ void FTraceWriter::Reset()
 	//TraceGuid
 	//SessionGuid
 
-	StartCycle = 0;
+	//////////////////////////////////////////////////
+
 	CycleFrequency = 0;
-	StartDateTime = 0.0;
+	TimeGetter.Reset();
+	StartCycle = 0;
+	StartTimeSinceEpoch = 0.0;
+
+	bUseCustomClock = false;
+	InitDefaultClock();
+
+	//////////////////////////////////////////////////
 
 	NextSerial = 0;
 
@@ -479,29 +489,7 @@ void FTraceWriter::GenerateGuid(FTraceWriter::FTraceGuid* OutGuid)
 
 ////////////////////////////////////////////////////////////////////////////////////////////////////
 
-uint64 FTraceWriter::GetTimeFrequency() const
-{
-	return FPlatformTime::SecondsToCycles64(1.0);
-}
-
-////////////////////////////////////////////////////////////////////////////////////////////////////
-
-uint64 FTraceWriter::GetTimeCycles64() const
-{
-	return FPlatformTime::Cycles64();
-}
-
-////////////////////////////////////////////////////////////////////////////////////////////////////
-
-double FTraceWriter::GetTimeSinceEpoch() const
-{
-	//return std::chrono::duration_cast<std::chrono::duration<double>>(std::chrono::system_clock::now().time_since_epoch()).count();
-	return FDateTime::Now().ToUnixTimestampDecimal();
-}
-
-////////////////////////////////////////////////////////////////////////////////////////////////////
-
-void FTraceWriter::Begin()
+void FTraceWriter::Begin(ETraceWriterBeginOptions Options)
 {
 	// Create the default threads.
 	Threads.Reset();
@@ -511,12 +499,13 @@ void FTraceWriter::Begin()
 	static_assert(ETransportTid::Importants == 1);
 	check(Threads.Num() == ETransportTid::Bias);
 
-	StartDateTime = GetTimeSinceEpoch();
-	StartCycle = GetTimeCycles64();
-	CycleFrequency = GetTimeFrequency();
-
 	constexpr uint16 GControlPort = 0;
 
+	if (!bUseCustomClock)
+	{
+		InitDefaultClock();
+	}
+
 	// Handshake
 	struct FHandshake
 	{
@@ -565,19 +554,18 @@ void FTraceWriter::Begin()
 		return;
 	}
 
-	// Describe initial events.
-	RegisterDefaultEvents();
-	WritePacket();
-
-	// Write $Trace.NewTrace event.
-	SetCurrentThread(static_cast<uint32>(ETransportTid::Importants));
-	WriteEvent(NewTraceEventId)
-		.Field(0, StartCycle)
-		.Field(1, CycleFrequency)
-		.Field(2, uint16(0x524d)) // Endian
-		.Field(3, uint8(sizeof(void*))) // PointerSize
-		.Field(4, StartDateTime)
-		.End();
+	if (EnumHasAnyFlags(Options, ETraceWriterBeginOptions::DeclareDefaultEvents))
+	{
+		// Describe initial events.
+		DeclareDefaultEvents();
+		WritePacket();
+	}
+	
+	if (EnumHasAnyFlags(Options, ETraceWriterBeginOptions::WriteNewTraceEvent))
+	{
+		// Write $Trace.NewTrace event.
+		WriteNewTraceEvent();
+	}
 }
 
 ////////////////////////////////////////////////////////////////////////////////////////////////////
@@ -590,7 +578,7 @@ void FTraceWriter::End()
 
 ////////////////////////////////////////////////////////////////////////////////////////////////////
 
-void FTraceWriter::RegisterDefaultEvents()
+void FTraceWriter::DeclareDefaultEvents()
 {
 	NewTraceEventId = DeclareEvent(ANSITEXTVIEW("$Trace"), ANSITEXTVIEW("NewTrace"), ETraceWriterEventFlags::ImportantNoSync)
 		.Field(ANSITEXTVIEW("StartCycle"), ETraceWriterFieldType::Uint64)
@@ -621,6 +609,27 @@ void FTraceWriter::RegisterDefaultEvents()
 
 ////////////////////////////////////////////////////////////////////////////////////////////////////
 
+void FTraceWriter::WriteNewTraceEvent()
+{
+	if (NewTraceEventId == 0)
+	{
+		SetError("$Trace.NewTrace event was not declared!");
+		return;
+	}
+
+	SetCurrentThread(static_cast<uint32>(ETransportTid::Importants));
+
+	WriteEvent(NewTraceEventId) // ImportantNoSync
+		.Field(0, StartCycle)
+		.Field(1, CycleFrequency)
+		.Field(2, uint16(0x524d)) // Endian
+		.Field(3, uint8(sizeof(void*))) // PointerSize
+		.Field(4, StartTimeSinceEpoch)
+		.End();
+}
+
+////////////////////////////////////////////////////////////////////////////////////////////////////
+
 uint32 FTraceWriter::DeclareDiagnosticsSession2Event()
 {
 	return DeclareEvent(ANSITEXTVIEW("Diagnostics"), ANSITEXTVIEW("Session2"), ETraceWriterEventFlags::ImportantNoSync)
@@ -641,7 +650,7 @@ uint32 FTraceWriter::DeclareDiagnosticsSession2Event()
 
 ////////////////////////////////////////////////////////////////////////////////////////////////////
 
-uint32 FTraceWriter::GetOrRegisterUnknownEventId()
+uint32 FTraceWriter::GetOrDeclareUnknownEventId()
 {
 	if (UnknownEventId == 0)
 	{
@@ -653,34 +662,85 @@ uint32 FTraceWriter::GetOrRegisterUnknownEventId()
 
 ////////////////////////////////////////////////////////////////////////////////////////////////////
 
-uint32 FTraceWriter::DeclareThread(FAnsiStringView ThreadName, uint32 SystemId, int32 SortHint)
+uint32 FTraceWriter::RegisterThread(FAnsiStringView ThreadName, uint32 SystemId, int32 SortHint, bool bShouldWriteThreadInfo)
 {
-	const FThreadInfo* FoundThread = Threads.FindByPredicate(
-		[&ThreadName](const FThreadInfo& ThreadInfo)
+	const FTraceWriterThreadInfo* FoundThread = Threads.FindByPredicate(
+		[&ThreadName](const FTraceWriterThreadInfo& ThreadInfo)
 		{
-			return ThreadName.Equals(ThreadInfo.Name, ESearchCase::IgnoreCase);
+			return ThreadInfo.Id != FTraceWriterThreadInfo::InvalidThreadId &&
+				ThreadName.Equals(ThreadInfo.Name, ESearchCase::CaseSensitive);
 		});
 	if (FoundThread)
 	{
 		return FoundThread->Id;
 	}
 
-	SetCurrentThreadImportants();
+	const uint32 ThreadId = uint32(Threads.Num());
+	RegisterCustomThread(ThreadId, ThreadName, SystemId, SortHint, bShouldWriteThreadInfo);
+	return ThreadId;
+}
 
-	FThreadInfo& Thread = Threads.AddDefaulted_GetRef();
-	Thread.Id = uint32(Threads.Num()) - 1;
+////////////////////////////////////////////////////////////////////////////////////////////////////
+
+void FTraceWriter::RegisterCustomThread(uint32 ThreadId, FAnsiStringView ThreadName, uint32 SystemId, int32 SortHint, bool bShouldWriteThreadInfo)
+{
+	const uint32 NumThreads = uint32(Threads.Num());
+
+	if (ThreadId < NumThreads && Threads[ThreadId].Id != FTraceWriterThreadInfo::InvalidThreadId)
+	{
+		SetError("Thread already registered!");
+		return;
+	}
+
+	if (ThreadId >= NumThreads)
+	{
+		Threads.AddDefaulted(ThreadId - NumThreads + 1);
+	}
+
+	FTraceWriterThreadInfo& Thread = Threads[ThreadId];
+	Thread.Id = ThreadId;
 	Thread.Name = FAnsiString(ThreadName);
 	Thread.SystemId = SystemId;
 	Thread.SortHint = SortHint;
 
+	if (bShouldWriteThreadInfo)
+	{
+		WriteThreadInfo(Thread.Id, Thread.Name, Thread.SystemId, Thread.SortHint);
+	}
+}
+
+////////////////////////////////////////////////////////////////////////////////////////////////////
+
+void FTraceWriter::WriteThreadInfo(uint32 ThreadId)
+{
+	if (ThreadId >= uint32(Threads.Num()))
+	{
+		SetError("Thread is not registered!");
+		return;
+	}
+
+	const FTraceWriterThreadInfo& Thread = Threads[ThreadId];
+	WriteThreadInfo(Thread.Id, Thread.Name, Thread.SystemId, Thread.SortHint);
+}
+
+////////////////////////////////////////////////////////////////////////////////////////////////////
+
+void FTraceWriter::WriteThreadInfo(uint32 ThreadId, FAnsiStringView ThreadName, uint32 SystemId, int32 SortHint)
+{
+	if (ThreadInfoEventId == 0)
+	{
+		SetError("$Trace.ThreadInfo event was not declared!");
+		return;
+	}
+
+	SetCurrentThreadImportants();
+
 	WriteEvent(ThreadInfoEventId) // ImportantNoSync
-		.Field(0, Thread.Id)
-		.Field(1, Thread.SystemId)
-		.Field(2, Thread.SortHint)
-		.Field(3, Thread.Name)
+		.Field(0, ThreadId)
+		.Field(1, SystemId)
+		.Field(2, SortHint)
+		.Field(3, ThreadName)
 		.End();
-
-	return Thread.Id;
 }
 
 ////////////////////////////////////////////////////////////////////////////////////////////////////
@@ -700,6 +760,13 @@ void FTraceWriter::SetCurrentThread(uint32 ThreadId)
 			return;
 		}
 
+		if (ThreadId >= uint32(Threads.Num()) ||
+			Threads[ThreadId].Id == FTraceWriterThreadInfo::InvalidThreadId)
+		{
+			SetError(TEXT("Thread is not registered!"));
+			return;
+		}
+
 		CurrentThreadId = ThreadId;
 	}
 }
@@ -708,6 +775,12 @@ void FTraceWriter::SetCurrentThread(uint32 ThreadId)
 
 void FTraceWriter::WriteThreadTimingEvent(uint64 BaseTimestamp)
 {
+	if (ThreadTimingEventId == 0)
+	{
+		SetError("$Trace.ThreadTiming event was not declared!");
+		return;
+	}
+
 	WriteEvent(ThreadTimingEventId) // NoSync
 		.Field(0, BaseTimestamp)
 		.End();
@@ -717,7 +790,14 @@ void FTraceWriter::WriteThreadTimingEvent(uint64 BaseTimestamp)
 
 void FTraceWriter::WriteThreadGroupBeginEvent(FAnsiStringView GroupName)
 {
+	if (ThreadGroupBeginEventId == 0)
+	{
+		SetError("$Trace.ThreadGroupBegin event was not declared!");
+		return;
+	}
+
 	SetCurrentThreadImportants();
+
 	WriteEvent(ThreadGroupBeginEventId) // ImportantNoSync
 		.Field(0, GroupName)
 		.End();
@@ -727,13 +807,36 @@ void FTraceWriter::WriteThreadGroupBeginEvent(FAnsiStringView GroupName)
 
 void FTraceWriter::WriteThreadGroupEndEvent()
 {
+	if (ThreadGroupEndEventId == 0)
+	{
+		SetError("$Trace.ThreadGroupEnd event was not declared!");
+		return;
+	}
+
 	SetCurrentThreadImportants();
+
 	WriteEvent(ThreadGroupEndEventId) // ImportantNoSync
 		.End();
 }
 
 ////////////////////////////////////////////////////////////////////////////////////////////////////
 
+const FTraceWriterEventInfo* FTraceWriter::FindEvent(FAnsiStringView LoggerName, FAnsiStringView EventName) const
+{
+	for (const FTraceWriterEventInfo* EventInfo : EventInfos)
+	{
+		if (EventInfo &&
+			LoggerName.Equals(EventInfo->LoggerName, ESearchCase::CaseSensitive) &&
+			EventName.Equals(EventInfo->Name, ESearchCase::CaseSensitive))
+		{
+			return EventInfo;
+		}
+	}
+	return nullptr;
+}
+
+////////////////////////////////////////////////////////////////////////////////////////////////////
+
 FTraceWriterEventDeclarationBuilder& FTraceWriter::DeclareEvent(FAnsiStringView LoggerName, FAnsiStringView Name, ETraceWriterEventFlags Flags)
 {
 	check(PendingEventDeclaration == nullptr);
@@ -958,7 +1061,7 @@ FTraceWriterEventBuilder& FTraceWriter::WriteEvent(uint32 EventId)
 	if (EventId >= static_cast<uint32>(EventInfos.Num()))
 	{
 		SetError(TEXT("Invalid event id!"));
-		EventId = GetOrRegisterUnknownEventId();
+		EventId = GetOrDeclareUnknownEventId();
 	}
 
 	check(EventId < static_cast<uint32>(EventInfos.Num()));
@@ -1251,23 +1354,23 @@ void FTraceWriter::WriteLeaveScopeEvent()
 
 ////////////////////////////////////////////////////////////////////////////////////////////////////
 
-void FTraceWriter::WriteStampedEnterScopeEvent()
+void FTraceWriter::WriteStampedEnterScopeEvent(uint64 Timestamp)
 {
 	BeginEvent();
 	using namespace Private::WriterProtocol;
 	constexpr uint8 Uid = uint8(EKnownEventUids::EnterScope_TB << EKnownEventUids::_UidShift);
-	WriteStampedScopeEventPrivate(Uid);
+	WriteStampedScopeEventPrivate(Uid, Timestamp);
 	EndEvent();
 }
 
 ////////////////////////////////////////////////////////////////////////////////////////////////////
 
-void FTraceWriter::WriteStampedLeaveScopeEvent()
+void FTraceWriter::WriteStampedLeaveScopeEvent(uint64 Timestamp)
 {
 	BeginEvent();
 	using namespace Private::WriterProtocol;
 	constexpr uint8 Uid = uint8(EKnownEventUids::LeaveScope_TB << EKnownEventUids::_UidShift);
-	WriteStampedScopeEventPrivate(Uid);
+	WriteStampedScopeEventPrivate(Uid, Timestamp);
 	EndEvent();
 }
 
@@ -1289,7 +1392,7 @@ void FTraceWriter::WriteScopeEventPrivate(uint8 Uid)
 
 ////////////////////////////////////////////////////////////////////////////////////////////////////
 
-void FTraceWriter::WriteStampedScopeEventPrivate(uint8 Uid)
+void FTraceWriter::WriteStampedScopeEventPrivate(uint8 Uid, uint64 Timestamp)
 {
 	constexpr uint32 Size = sizeof(uint64); // 1 byte (Uid) + 7 bytes (timestamp)
 
@@ -1300,7 +1403,7 @@ void FTraceWriter::WriteStampedScopeEventPrivate(uint8 Uid)
 		return;
 	}
 
-	const uint64 RelativeTimestamp = GetTimeCycles64() - GetStartTimeCycles64();
+	const uint64 RelativeTimestamp = Timestamp - GetStartTime();
 	if ((RelativeTimestamp >> 56) != 0)
 	{
 		SetError(TEXT("Relative timestamp is too large!"));
@@ -1468,6 +1571,10 @@ ETraceWriterFieldType FTraceWriter::ConvertFieldType(const UE::Trace::IAnalyzer:
 					case 8: WriterFieldType = ETraceWriterFieldType::Uint64; break;
 				}
 			}
+			if (FieldInfo.IsArray())
+			{
+				EnumAddFlags(WriterFieldType, ETraceWriterFieldType::ArrayFlag);
+			}
 			break;
 
 		case EType::Float:
@@ -1480,6 +1587,10 @@ ETraceWriterFieldType FTraceWriter::ConvertFieldType(const UE::Trace::IAnalyzer:
 				check(FieldInfo.GetTypeSize() == 8);
 				WriterFieldType = ETraceWriterFieldType::Float64;
 			}
+			if (FieldInfo.IsArray())
+			{
+				EnumAddFlags(WriterFieldType, ETraceWriterFieldType::ArrayFlag);
+			}
 			break;
 
 		case EType::AnsiString: WriterFieldType = ETraceWriterFieldType::AnsiString; break;
@@ -1491,11 +1602,6 @@ ETraceWriterFieldType FTraceWriter::ConvertFieldType(const UE::Trace::IAnalyzer:
 		case EType::Reference64: WriterFieldType = ETraceWriterFieldType::ReferenceFlag | ETraceWriterFieldType::Uint64; break;
 	}
 
-	if (FieldInfo.IsArray())
-	{
-		EnumAddFlags(WriterFieldType, ETraceWriterFieldType::ArrayFlag);
-	}
-
 	return WriterFieldType;
 }
 
@@ -1615,6 +1721,21 @@ void FTraceWriter::WritePacket()
 
 ////////////////////////////////////////////////////////////////////////////////////////////////////
 
+void FTraceWriter::InitDefaultClock()
+{
+	check(!bUseCustomClock);
+
+	CycleFrequency = FPlatformTime::SecondsToCycles64(1.0);
+	TimeGetter = []() { return FPlatformTime::Cycles64(); };
+
+	StartCycle = FPlatformTime::Cycles64();
+
+	//std::chrono::duration_cast<std::chrono::duration<double>>(std::chrono::system_clock::now().time_since_epoch()).count()
+	StartTimeSinceEpoch = FDateTime::Now().ToUnixTimestampDecimal();
+}
+
+////////////////////////////////////////////////////////////////////////////////////////////////////
+
 void FTraceWriter::SetError(FString&& InError)
 {
 	Errors.Push(InError);
diff --git a/Engine/Source/Developer/TraceAnalysis/Public/Trace/TraceWriter.h b/Engine/Source/Developer/TraceAnalysis/Public/Trace/TraceWriter.h
--- a/Engine/Source/Developer/TraceAnalysis/Public/Trace/TraceWriter.h
+++ b/Engine/Source/Developer/TraceAnalysis/Public/Trace/TraceWriter.h
@@ -6,6 +6,7 @@
 #include "Containers/AnsiString.h"
 #include "Containers/UnrealString.h"
 #include "Misc/EnumClassFlags.h"
+#include "Templates/Function.h"
 #include "Templates/Requires.h"
 
 // TraceLog
@@ -25,6 +26,36 @@ namespace UE::Trace
 class IOutDataStream;
 class FTraceWriter;
 
+////////////////////////////////////////////////////////////////////////////////////////////////////
+
+class FTraceWriterThreadInfo final
+{
+	friend class FTraceWriter;
+
+public:
+	FTraceWriterThreadInfo() = default;
+	FTraceWriterThreadInfo(uint32 InThreadId, FAnsiStringView InName, uint32 InSystemId = 0, int32 InSortHint = 0)
+		: Id(InThreadId)
+		, Name(InName)
+		, SystemId(InSystemId)
+		, SortHint(InSortHint)
+	{
+	}
+
+	[[nodiscard]] uint32 GetId() const { return Id; }
+	[[nodiscard]] const FAnsiString& GetName() const { return Name; }
+	[[nodiscard]] uint32 GetSystemId() const { return SystemId; }
+	[[nodiscard]] uint32 GetSortHint() const { return SortHint; }
+
+	static constexpr uint32 InvalidThreadId = ~0u;
+
+private:
+	uint32 Id = InvalidThreadId;
+	FAnsiString Name;
+	uint32 SystemId = 0;
+	int32 SortHint = 0;
+};
+
 ////////////////////////////////////////////////////////////////////////////////////////////////////
 /**
  * Flags for a trace event.
@@ -534,6 +565,19 @@ private:
 
 ////////////////////////////////////////////////////////////////////////////////////////////////////
 
+enum class ETraceWriterBeginOptions : uint32
+{
+	None = 0,
+
+	DeclareDefaultEvents = 1 << 1,
+	WriteNewTraceEvent = 1 << 2,
+
+	Default = DeclareDefaultEvents | WriteNewTraceEvent
+};
+ENUM_CLASS_FLAGS(ETraceWriterBeginOptions);
+
+////////////////////////////////////////////////////////////////////////////////////////////////////
+
 class FTraceWriter
 {
 	friend class FTraceWriterEventInfo;
@@ -545,13 +589,6 @@ private:
 	{
 		uint32 Bits[4];
 	};
-	struct FThreadInfo
-	{
-		uint32 Id;
-		FAnsiString Name;
-		uint32 SystemId = 0;
-		int32 SortHint = 0;
-	};
 
 public:
 	UE_API FTraceWriter(IOutDataStream& InStream);
@@ -576,12 +613,11 @@ public:
 
 	/**
 	 * Begins the trace stream.
+	 * Calls DeclareDefaultEvents() and WriteNewTraceEvent() depending on Options flags.
 	 *
-	 * \details
-	 * Declares: $Trace.* events
-	 * Writes: $Trace.NewEvent, $Trace.NewTrace
+	 * @param Options - Various toggle options to allow custom control of Begin() behavior
 	 */
-	UE_API void Begin();
+	UE_API void Begin(ETraceWriterBeginOptions Options = ETraceWriterBeginOptions::Default);
 
 	/**
 	 * Ends the trace stream.
@@ -592,38 +628,89 @@ public:
 	UE_API void End();
 
 	/**
-	 * Writes a $Trace.ThreadGroupBegin event.
+	 * Declares the default trace events.
+	 *
+	 * \details
+	 * Declares: $Trace.* events
+	 */
+	UE_API void DeclareDefaultEvents();
+
+	/**
+	 * Writes the $Trace.NewTrace event.
 	 * Changes the current thread to be the "Importants" thread.
 	 *
 	 * \details
-	 * Writes: $Trace.ThreadGroupBegin
+	 * Writes: $Trace.NewTrace
+	 */
+	UE_API void WriteNewTraceEvent();
+
+	//////////////////////////////////////////////////
+	// Threads
+
+	/**
+	 * Registers a new trace thread.
+	 * If a thread with same name is already registered, this will just return the registered thread id.
+	 * If successfully registering a new thread and bShouldWriteThreadInfo is true, then this function
+	 * will call WriteThreadInfo() automatically.
 	 *
-	 * @param GroupName - The group name
+	 * @param ThreadName - The thread name
+	 * @param SystemId - The system thread id
+	 * @param SortHint - A hint for sorting priority
+	 * @param bShouldWriteThreadInfo - Wherever should write the $Trace.ThreadInfo or not
+	 *
+	 * @returns the thread id for the new trace thread
 	 */
-	UE_API void WriteThreadGroupBeginEvent(FAnsiStringView GroupName);
+	UE_API uint32 RegisterThread(FAnsiStringView ThreadName, uint32 SystemId = 0, int32 SortHint = 0, bool bShouldWriteThreadInfo = true);
 
 	/**
-	 * Writes a $Trace.ThreadGroupEnd event.
+	 * Registers a trace thread with a known id.
+	 * If a thread with the specified id is already registered, this will fail.
+	 * If successfully registering a new thread and bShouldWriteThreadInfo is true, then this function
+	 * will call WriteThreadInfo() automatically.
+	 *
+	 * @param ThreadId - The thread id
+	 * @param ThreadName - The thread name
+	 * @param SystemId - The system thread id
+	 * @param SortHint - A hint for sorting priority
+	 * @param bShouldWriteThreadInfo - Wherever should write the $Trace.ThreadInfo or not
+	 */
+	UE_API void RegisterCustomThread(uint32 ThreadId, FAnsiStringView ThreadName, uint32 SystemId = 0, int32 SortHint = 0, bool bShouldWriteThreadInfo = true);
+
+	/**
+	 * Gets the info for a registered thread identified by Id.
+	 *
+	 * @param ThreadId - The thread id
+	 * @returns a pointer to the thread info or nullptr if the ThreadId is not valid (thread not registered).
+	 */
+	const FTraceWriterThreadInfo* GetThreadInfo(uint32 ThreadId)
+	{
+		return ThreadId < uint32(Threads.Num()) && Threads[ThreadId].Id != FTraceWriterThreadInfo::InvalidThreadId ? &Threads[ThreadId] : nullptr;
+	}
+
+	/**
+	 * Writes the $Trace.ThreadInfo event for a registered thread, identified by id.
 	 * Changes the current thread to be the "Importants" thread.
 	 *
 	 * \details
-	 * Writes: $Trace.ThreadGroupEnd
+	 * Writes: $Trace.ThreadInfo
+	 *
+	 * @param ThreadId - The thread id
 	 */
-	UE_API void WriteThreadGroupEndEvent();
+	UE_API void WriteThreadInfo(uint32 ThreadId);
 
 	/**
-	 * Declares a new trace thread.
 	 * Writes the $Trace.ThreadInfo event.
 	 * Changes the current thread to be the "Importants" thread.
 	 *
 	 * \details
 	 * Writes: $Trace.ThreadInfo
 	 *
+	 * @param ThreadId - The thread id
 	 * @param ThreadName - The thread name
 	 * @param SystemId - The system thread id
 	 * @param SortHint - A hint for sorting priority
 	 */
-	UE_API uint32 DeclareThread(FAnsiStringView ThreadName, uint32 SystemId = 0, int32 SortHint = 0);
+	UE_API void WriteThreadInfo(uint32 ThreadId, FAnsiStringView ThreadName, uint32 SystemId = 0, int32 SortHint = 0);
 
 	/**
 	 * Sets the current thread id.
@@ -652,10 +739,52 @@ public:
 	/**
 	 * Writes a $Trace.ThreadTiming event for the current thread.
 	 *
-	 * @param BaseTimestamp - The base timestamp [Cycles64]
+	 * @param BaseTimestamp - The base timestamp [ticks]
 	 */
 	UE_API void WriteThreadTimingEvent(uint64 BaseTimestamp);
 
+	/**
+	 * Writes a $Trace.ThreadGroupBegin event.
+	 * Changes the current thread to be the "Importants" thread.
+	 *
+	 * \details
+	 * Writes: $Trace.ThreadGroupBegin
+	 *
+	 * @param GroupName - The group name
+	 */
+	UE_API void WriteThreadGroupBeginEvent(FAnsiStringView GroupName);
+
+	/**
+	 * Writes a $Trace.ThreadGroupEnd event.
+	 * Changes the current thread to be the "Importants" thread.
+	 *
+	 * \details
+	 * Writes: $Trace.ThreadGroupEnd
+	 */
+	UE_API void WriteThreadGroupEndEvent();
+
+	//////////////////////////////////////////////////
+
+	/**
+	 * Gets the info for an event identified by Id.
+	 *
+	 * @param EventId - The event id, one returned by DeclareEvent().[...].End() or similar method.
+	 * @returns a pointer to the event info or nullptr if the EventId is not valid.
+	 */
+	const FTraceWriterEventInfo* GetEventInfo(uint32 EventId) const
+	{
+		return EventId < uint32(EventInfos.Num()) ? EventInfos[EventId] : nullptr;
+	}
+
+	/**
+	 * Gets the info for an event identified by logger and event names.
+	 *
+	 * @param LoggerName - The logger name
+	 * @param Name - The event name
+	 * @returns a pointer to the event info or nullptr if the event is not found.
+	 */
+	UE_API const FTraceWriterEventInfo* FindEvent(FAnsiStringView LoggerName, FAnsiStringView EventName) const;
+
 	/**
 	 * Declares a new trace event.
 	 * The End() call for the returned builder will change the current thread to be the "Events" thread.
@@ -674,17 +803,6 @@ public:
 	 */
 	[[nodiscard]] UE_API FTraceWriterEventDeclarationBuilder& DeclareEvent(FAnsiStringView LoggerName, FAnsiStringView Name, ETraceWriterEventFlags Flags = ETraceWriterEventFlags::None);
 
-	/**
-	 * Gets the info for an event identified by Id.
-	 *
-	 * @param EventId - The event id, one returned by DeclareEvent().[...].End() or similar method.
-	 * @returns a pointer to the event info or nullptr if the EventId is not valid.
-	 */
-	const FTraceWriterEventInfo* GetEventInfo(uint32 EventId) const
-	{
-		return EventId < uint32(EventInfos.Num()) ? EventInfos[EventId] : nullptr;
-	}
-
 	/**
 	 * Declares the Diagnostics.Session2 event.
 	 *
@@ -707,11 +825,28 @@ public:
 	//////////////////////////////////////////////////
 	// Scoped Events
 
+	/**
+	 * Writes the EnterScope event, on the current thread.
+	 */
 	UE_API void WriteEnterScopeEvent();
+	/**
+	 * Writes the LeaveScope event, on the current thread.
+	 */
 	UE_API void WriteLeaveScopeEvent();
 
-	UE_API void WriteStampedEnterScopeEvent();
-	UE_API void WriteStampedLeaveScopeEvent();
+	/**
+	 * Writes the EnterScope_TB event, on the current thread, using the specified timestamp.
+	 *
+	 * @param Timestamp - The current session time (i.e. one returned by GetTime()).
+	 */
+	UE_API void WriteStampedEnterScopeEvent(uint64 Timestamp);
+
+	/**
+	 * Writes the LeaveScope_TB event, on the current thread, using the specified timestamp.
+	 *
+	 * @param Timestamp - The current session time (i.e. one returned by GetTime()).
+	 */
+	UE_API void WriteStampedLeaveScopeEvent(uint64 Timestamp);
 
 	//////////////////////////////////////////////////
 	// Event Flags
@@ -763,17 +898,76 @@ public:
 	[[nodiscard]] UE_API static ETraceWriterFieldType ConvertFieldType(const UE::Trace::IAnalyzer::FEventFieldInfo& FieldInfo);
 
 	//////////////////////////////////////////////////
+	// Time
 
-	[[nodiscard]] UE_API uint64 GetTimeFrequency() const;
-	[[nodiscard]] UE_API uint64 GetTimeCycles64() const;
-	[[nodiscard]] UE_API double GetTimeSinceEpoch() const;
+	/**
+	 * Sets the custom clock for the session time.
+	 * Warning: It is only valid to set these before writer's Begin().
+	 *
+	 * @param InTimeFrequency - The time frequency; the number of [ticks] in a second.
+	 * @param InTimeGetter - The custom getter for the session time; when called it is expected to return the current session time, in [ticks].
+	 * @param InStartTime - The start session timestamp, in [ticks].
+	 * @param InStartTimeSinceEpoch - The start session time since Unix Epoch (January 1st of 1970), in [seconds]. Should match InStartTime.
+	 */
+	void SetCustomClock(uint64 InTimeFrequency, TFunction<uint64()> InTimeGetter, uint64 InStartTime, double InStartTimeSinceEpoch)
+	{
+		CycleFrequency = InTimeFrequency;
+		TimeGetter = InTimeGetter;
+		StartCycle = InStartTime;
+		StartTimeSinceEpoch = InStartTimeSinceEpoch;
+		bUseCustomClock = true;
+	}
 
-	//UE_API void SetCustomStartTime(uint64 InStartTimeFrequency, uint64 InStartTimeCycles64, double InStartTimeSinceEpoch);
-	//UE_API void SetCustomClock(TFunctionRef<uint64> TimeFrequencyGetter, TFunctionRef<uint64> TimeCycles64Getter, TFunctionRef<double> TimeSinceEpochGetter);
+	/**
+	 * Returns the frequency of the session time, specified as number of [ticks] per second.
+	 *
+	 * @returns the frequency of the session time.
+	 */
+	[[nodiscard]] uint64 GetTimeFrequency() const
+	{
+		return CycleFrequency;
+	}
 
-	[[nodiscard]] uint64 GetStartTimeFrequency() const { return CycleFrequency; }
-	[[nodiscard]] uint64 GetStartTimeCycles64() const { return StartCycle; }
-	[[nodiscard]] double GetStartTimeSinceEpoch() const { return StartDateTime; }
+	/**
+	 * Gets the start session time, in [ticks].
+	 *
+	 * @returns the start session time.
+	 */
+	[[nodiscard]] uint64 GetStartTime() const
+	{
+		return StartCycle;
+	}
+
+	/**
+	 * Gets the start session time since the Unix Epoch (January 1st of 1970), in [seconds].
+	 *
+	 * @returns the start session time since the Unix Epoch.
+	 */
+	[[nodiscard]] double GetStartTimeSinceEpoch() const
+	{
+		return StartTimeSinceEpoch;
+	}
+
+	/**
+	 * Gets the current session time, in [ticks].
+	 *
+	 * @returns the current session time.
+	 */
+	[[nodiscard]] UE_API uint64 GetTime() const
+	{
+		return TimeGetter();
+	}
+
+	/**
+	 * Gets the current session time since the Unix Epoch (January 1st of 1970), in [seconds].
+	 *
+	 * @returns the current session time since the Unix Epoch.
+	 */
+	[[nodiscard]] double GetTimeSinceEpoch() const
+	{
+		const uint64 RelativeTime = GetTime() - StartCycle;
+		return StartTimeSinceEpoch + double(RelativeTime) / double(CycleFrequency);
+	}
 
 	//////////////////////////////////////////////////
 
@@ -796,8 +990,7 @@ public:
 private:
 	void Reset();
 	static void GenerateGuid(FTraceGuid* OutGuid);
-	void RegisterDefaultEvents();
-	uint32 GetOrRegisterUnknownEventId();
+	uint32 GetOrDeclareUnknownEventId();
 
 	void WriteEventDeclaration(const FTraceWriterEventDeclarationBuilder& Builder);
 	void WriteEventDeclaration(const FTraceWriterEventInfo& EventInfo);
@@ -807,7 +1000,7 @@ private:
 	void WriteAuxDataSegment(uint32 Index, const void* SegmentData, uint32 SegmentSize, bool bIsImportant);
 
 	void WriteScopeEventPrivate(uint8 Uid);
-	void WriteStampedScopeEventPrivate(uint8 Uid);
+	void WriteStampedScopeEventPrivate(uint8 Uid, uint64 Timestamp);
 
 	uint8* GetBuffer(uint32 RequiredDataSize);
 	void AdvanceBuffer(uint32 DataSize);
@@ -815,6 +1008,7 @@ private:
 	void EndEvent();
 	void WritePacket();
 
+	void InitDefaultClock();
 	void SetError(FString&& InError);
 
 private:
@@ -823,9 +1017,11 @@ private:
 	FTraceGuid TraceGuid;
 	FTraceGuid SessionGuid;
 
-	uint64 StartCycle = 0;
-	uint64 CycleFrequency = 0;
-	double StartDateTime = 0.0;
+	uint64 CycleFrequency = 0; // [ticks] per second
+	TFunction<uint64()> TimeGetter; // returns session time, in [ticks]
+	uint64 StartCycle = 0; // [ticks]
+	double StartTimeSinceEpoch = 0.0; // [seconds] since epoch
+	bool bUseCustomClock = false;
 
 	uint32 NextSerial = 0;
 
@@ -836,7 +1032,7 @@ private:
 	uint32 ThreadGroupBeginEventId = 0;
 	uint32 ThreadGroupEndEventId = 0;
 
-	TArray<FThreadInfo> Threads;
+	TArray<FTraceWriterThreadInfo> Threads;
 
 	TArray<FTraceWriterEventInfo*> EventInfos;
 	FTraceWriterEventDeclarationBuilder* PendingEventDeclaration = nullptr;
