[Trace Analysis] Deprecated FStoreClient::Connect API
- Introduced a new API that uses a uniqueptr to prevent mem leaks
- Updated deprecated usages

#rnx
#Synced-CL 49728215

diff --git a/Engine/Plugins/Animation/GameplayInsights/Source/RewindDebugger/Private/SModalSessionBrowser.cpp b/Engine/Plugins/Animation/GameplayInsights/Source/RewindDebugger/Private/SModalSessionBrowser.cpp
--- a/Engine/Plugins/Animation/GameplayInsights/Source/RewindDebugger/Private/SModalSessionBrowser.cpp
+++ b/Engine/Plugins/Animation/GameplayInsights/Source/RewindDebugger/Private/SModalSessionBrowser.cpp
@@ -238,7 +238,7 @@ bool SModalSessionBrowser::CanOpenSession() const
 void SModalSessionBrowser::UpdateCurrentSessionInfoMap()
 {
 	using namespace UE::Trace;
-	const FStoreClient* StoreClient = FStoreClient::Connect(*CurrentTraceStoreAddress);
+	const TUniquePtr<FStoreClient> StoreClient = FStoreClient::Create(*CurrentTraceStoreAddress);
 
 	if (!StoreClient)
 	{
diff --git a/Engine/Plugins/AudioInsights/Source/AudioInsights/Private/AudioInsightsTraceModule.cpp b/Engine/Plugins/AudioInsights/Source/AudioInsights/Private/AudioInsightsTraceModule.cpp
--- a/Engine/Plugins/AudioInsights/Source/AudioInsights/Private/AudioInsightsTraceModule.cpp
+++ b/Engine/Plugins/AudioInsights/Source/AudioInsights/Private/AudioInsightsTraceModule.cpp
@@ -565,7 +565,7 @@ namespace UE::Audio::Insights
 			};
 			
 			
-			const TUniquePtr<UE::Trace::FStoreClient> StoreClient(UE::Trace::FStoreClient::Connect(*GetTraceHostAddress()));
+			const TUniquePtr<UE::Trace::FStoreClient> StoreClient = UE::Trace::FStoreClient::Create(*GetTraceHostAddress());
 			if (StoreClient.IsValid())
 			{
 				if (const UE::Trace::FStoreClient::FStatus* StoreClientStatus = StoreClient->GetStatus())
diff --git a/Engine/Plugins/ChaosVD/Source/ChaosVD/Private/Widgets/SChaosVDBrowseSessionsModal.cpp b/Engine/Plugins/ChaosVD/Source/ChaosVD/Private/Widgets/SChaosVDBrowseSessionsModal.cpp
--- a/Engine/Plugins/ChaosVD/Source/ChaosVD/Private/Widgets/SChaosVDBrowseSessionsModal.cpp
+++ b/Engine/Plugins/ChaosVD/Source/ChaosVD/Private/Widgets/SChaosVDBrowseSessionsModal.cpp
@@ -279,7 +279,7 @@ bool SChaosVDBrowseSessionsModal::CanOpenSession() const
 void SChaosVDBrowseSessionsModal::UpdateCurrentSessionInfoMap()
 {
 	using namespace UE::Trace;
-	const FStoreClient* StoreClient = FStoreClient::Connect(*CurrentTraceStoreAddress);
+	const TUniquePtr<FStoreClient> StoreClient = FStoreClient::Create(*CurrentTraceStoreAddress);
 
 	if (!StoreClient)
 	{
diff --git a/Engine/Plugins/Runtime/StateTree/Source/StateTreeModule/Private/StateTreeModule.cpp b/Engine/Plugins/Runtime/StateTree/Source/StateTreeModule/Private/StateTreeModule.cpp
--- a/Engine/Plugins/Runtime/StateTree/Source/StateTreeModule/Private/StateTreeModule.cpp
+++ b/Engine/Plugins/Runtime/StateTree/Source/StateTreeModule/Private/StateTreeModule.cpp
@@ -56,7 +56,7 @@ UE::Trace::FStoreClient* FStateTreeModule::GetStoreClient()
 {
 	if (!StoreClient.IsValid())
 	{
-		StoreClient = TUniquePtr<UE::Trace::FStoreClient>(UE::Trace::FStoreClient::Connect(TEXT("localhost")));
+		StoreClient = UE::Trace::FStoreClient::Create(TEXT("localhost"));
 	}
 	return StoreClient.Get();
 }
diff --git a/Engine/Plugins/TraceUtilities/Source/EditorTraceUtilities/Private/InsightsStatusBar/SInsightsStatusBar.cpp b/Engine/Plugins/TraceUtilities/Source/EditorTraceUtilities/Private/InsightsStatusBar/SInsightsStatusBar.cpp
--- a/Engine/Plugins/TraceUtilities/Source/EditorTraceUtilities/Private/InsightsStatusBar/SInsightsStatusBar.cpp
+++ b/Engine/Plugins/TraceUtilities/Source/EditorTraceUtilities/Private/InsightsStatusBar/SInsightsStatusBar.cpp
@@ -932,8 +932,7 @@ void SInsightsStatusBarWidget::CacheTraceStorePath()
 	if (TraceStorePath.IsEmpty())
 	{
 		using UE::Trace::FStoreClient;
-		FStoreClient* StoreClientPtr = FStoreClient::Connect(TEXT("localhost"));
-		TUniquePtr<FStoreClient> StoreClient = TUniquePtr<FStoreClient>(StoreClientPtr);
+		TUniquePtr<FStoreClient> StoreClient = FStoreClient::Create(TEXT("localhost"));
 
 		if (!StoreClient)
 		{
diff --git a/Engine/Plugins/TraceUtilities/Source/EditorTraceUtilities/Private/UnrealInsightsLauncher.cpp b/Engine/Plugins/TraceUtilities/Source/EditorTraceUtilities/Private/UnrealInsightsLauncher.cpp
--- a/Engine/Plugins/TraceUtilities/Source/EditorTraceUtilities/Private/UnrealInsightsLauncher.cpp
+++ b/Engine/Plugins/TraceUtilities/Source/EditorTraceUtilities/Private/UnrealInsightsLauncher.cpp
@@ -53,7 +53,7 @@ public:
 	void DoTask(ENamedThreads::Type CurrentThread, const FGraphEventRef& MyCompletionGraphEvent)
 	{
 		using namespace UE::Trace;
-		FStoreClient* StoreClient = FStoreClient::Connect(TEXT("localhost"));
+		TUniquePtr<FStoreClient> StoreClient = FStoreClient::Create(TEXT("localhost"));
 
 		OutData->TaskLiveSessionData.Empty();
 
@@ -186,7 +186,7 @@ void FUnrealInsightsLauncher::BuildUnrealInsights(const FString& Path, const FSt
 			if (Result.Equals(TEXT("Completed")))
 			{
 #if PLATFORM_MAC
-				// On Mac we genereate the path again so that it includes the newly built executable.
+				// On Mac we generate the path again so that it includes the newly built executable.
 				FString NewPath = GetInsightsApplicationPath();
 				this->StartUnrealInsights(NewPath, LaunchParameters, Callback);
 #else
@@ -277,7 +277,7 @@ bool FUnrealInsightsLauncher::OpenActiveTraceFromStore(const FString& TraceHostA
 {
 	TRACE_CPUPROFILER_EVENT_SCOPE(OpenActiveTraceFromStore);
 
-	UE::Trace::FStoreClient* StoreClient = UE::Trace::FStoreClient::Connect(*TraceHostAddress);
+	TUniquePtr<UE::Trace::FStoreClient> StoreClient = UE::Trace::FStoreClient::Create(*TraceHostAddress);
 
 	if (!StoreClient)
 	{
diff --git a/Engine/Source/Developer/TraceAnalysis/Private/Store/StoreClient.cpp b/Engine/Source/Developer/TraceAnalysis/Private/Store/StoreClient.cpp
--- a/Engine/Source/Developer/TraceAnalysis/Private/Store/StoreClient.cpp
+++ b/Engine/Source/Developer/TraceAnalysis/Private/Store/StoreClient.cpp
@@ -515,6 +515,22 @@ FStoreClient* FStoreClient::Connect(const TCHAR* Host, uint32 Port)
 	return (FStoreClient*)Impl;
 }
 
+////////////////////////////////////////////////////////////////////////////////
+TUniquePtr<FStoreClient> FStoreClient::Create(const TCHAR* Host, uint32 Port)
+{
+	static asio::io_context IoContext;
+
+	FStoreCborClient* Impl = new FStoreCborClient(IoContext);
+	if (!Impl->Connect(Host, static_cast<uint16>(Port)))
+	{
+		delete Impl;
+		return nullptr;
+	}
+
+	FStoreClient* Ptr = (FStoreClient*)Impl;
+	return TUniquePtr<FStoreClient>(Ptr);
+}
+
 ////////////////////////////////////////////////////////////////////////////////
 bool FStoreClient::Reconnect(const TCHAR* Host, uint32 Port)
 {
diff --git a/Engine/Source/Developer/TraceAnalysis/Private/Store/StoreConnection.cpp b/Engine/Source/Developer/TraceAnalysis/Private/Store/StoreConnection.cpp
--- a/Engine/Source/Developer/TraceAnalysis/Private/Store/StoreConnection.cpp
+++ b/Engine/Source/Developer/TraceAnalysis/Private/Store/StoreConnection.cpp
@@ -27,7 +27,7 @@ FStoreConnection::~FStoreConnection()
 
 bool FStoreConnection::ConnectToStore(const TCHAR* Host, uint32 Port)
 {
-	StoreClient.Reset(UE::Trace::FStoreClient::Connect(Host, Port));
+	StoreClient = FStoreClient::Create(Host, Port);
 	if (!StoreClient.IsValid())
 	{
 		return false;
diff --git a/Engine/Source/Developer/TraceAnalysis/Public/Trace/StoreClient.h b/Engine/Source/Developer/TraceAnalysis/Public/Trace/StoreClient.h
--- a/Engine/Source/Developer/TraceAnalysis/Public/Trace/StoreClient.h
+++ b/Engine/Source/Developer/TraceAnalysis/Public/Trace/StoreClient.h
@@ -21,29 +21,29 @@ class FStoreClient
 public:
 	struct FStatus
 	{
-		UE_API FUtf8StringView	GetStoreDir() const;
-		UE_API uint32			GetRecorderPort() const;
-		UE_API uint32			GetStorePort() const;
-		UE_API bool				GetSponsored() const;
-		UE_API uint32			GetChangeSerial() const;
-		UE_API uint32			GetSettingsSerial() const;
-		UE_API void				GetWatchDirectories(TArray<FString>& OutDirs) const;
+		UE_API FUtf8StringView				GetStoreDir() const;
+		UE_API uint32						GetRecorderPort() const;
+		UE_API uint32						GetStorePort() const;
+		UE_API bool							GetSponsored() const;
+		UE_API uint32						GetChangeSerial() const;
+		UE_API uint32						GetSettingsSerial() const;
+		UE_API void							GetWatchDirectories(TArray<FString>& OutDirs) const;
 	};
 
 	struct FVersion
 	{
-		UE_API uint32			GetMajorVersion() const;
-		UE_API uint32			GetMinorVersion() const;
-		UE_API FUtf8StringView	GetConfiguration() const;
+		UE_API uint32						GetMajorVersion() const;
+		UE_API uint32						GetMinorVersion() const;
+		UE_API FUtf8StringView				GetConfiguration() const;
 	};
 
 	struct FTraceInfo
 	{
-		UE_API FUtf8StringView	GetName() const;
-		UE_API uint32			GetId() const;
-		UE_API uint64			GetSize() const;
-		UE_API uint64			GetTimestamp() const;
-		UE_API FUtf8StringView	GetUri() const;
+		UE_API FUtf8StringView				GetName() const;
+		UE_API uint32						GetId() const;
+		UE_API uint64						GetSize() const;
+		UE_API uint64						GetTimestamp() const;
+		UE_API FUtf8StringView				GetUri() const;
 	};
 
 	struct FTraceData
@@ -52,41 +52,43 @@ public:
 		using TUniquePtr<IInDataStream>::TUniquePtr;
 	};
 
-								~FStoreClient() = default;
-	static UE_API FStoreClient* Connect(const TCHAR* Host, uint32 Port = 0);
-	UE_API bool					Reconnect(const TCHAR* Host, uint32 Port);
-	UE_API void					operator delete (void* Addr);
-	UE_API bool					IsValid() const;
-	UE_API uint32				GetStoreAddress() const;
-	UE_API uint32				GetStorePort() const;
-	UE_API const FStatus*		GetStatus() const;
-	UE_API const FVersion*		GetVersion() const;
-	UE_API uint32				GetTraceCount() const;
-	UE_API const FTraceInfo*	GetTraceInfo(uint32 Index) const;
-	UE_API const FTraceInfo*	GetTraceInfoById(uint32 Id) const;
-	UE_API FTraceData			ReadTrace(uint32 Id) const;
-	UE_API bool					SetStoreDirectories(const TCHAR* StoreDir, const TArray<FString>& AddWatchDirs, const TArray<FString>& RemoveWatchDirs);
-	UE_API bool					SetSponsored(bool bSponsored);
+											~FStoreClient() = default;
+	UE_DEPRECATED(5.8, "Connect has been deprecated. Please use Create which returns a TUniquePtr instead.")
+	static UE_API FStoreClient*				Connect(const TCHAR* Host, uint32 Port = 0);
+	static UE_API TUniquePtr<FStoreClient>  Create(const TCHAR* Host, uint32 Port = 0);
+	UE_API bool								Reconnect(const TCHAR* Host, uint32 Port);
+	UE_API void								operator delete (void* Addr);
+	UE_API bool								IsValid() const;
+	UE_API uint32							GetStoreAddress() const;
+	UE_API uint32							GetStorePort() const;
+	UE_API const FStatus*					GetStatus() const;
+	UE_API const FVersion*					GetVersion() const;
+	UE_API uint32							GetTraceCount() const;
+	UE_API const FTraceInfo*				GetTraceInfo(uint32 Index) const;
+	UE_API const FTraceInfo*				GetTraceInfoById(uint32 Id) const;
+	UE_API FTraceData						ReadTrace(uint32 Id) const;
+	UE_API bool								SetStoreDirectories(const TCHAR* StoreDir, const TArray<FString>& AddWatchDirs, const TArray<FString>& RemoveWatchDirs);
+	UE_API bool								SetSponsored(bool bSponsored);
 
 	struct FSessionInfo
 	{
-		UE_API uint32			GetId() const;
-		UE_API uint32			GetTraceId() const;
-		UE_API uint32			GetIpAddress() const;
-		UE_API uint32			GetControlPort() const;
+		UE_API uint32						GetId() const;
+		UE_API uint32						GetTraceId() const;
+		UE_API uint32						GetIpAddress() const;
+		UE_API uint32						GetControlPort() const;
 	};
-	UE_API uint32				GetSessionCount() const;
-	UE_API const FSessionInfo*	GetSessionInfo(uint32 Index) const;
-	UE_API const FSessionInfo*	GetSessionInfoById(uint32 Id) const;
-	UE_API const FSessionInfo*	GetSessionInfoByTraceId(uint32 TraceId) const;
-	UE_API const FSessionInfo*	GetSessionInfoByGuid(const FGuid& TraceGuid) const;
+	UE_API uint32							GetSessionCount() const;
+	UE_API const FSessionInfo*				GetSessionInfo(uint32 Index) const;
+	UE_API const FSessionInfo*				GetSessionInfoById(uint32 Id) const;
+	UE_API const FSessionInfo*				GetSessionInfoByTraceId(uint32 TraceId) const;
+	UE_API const FSessionInfo*				GetSessionInfoByGuid(const FGuid& TraceGuid) const;
 
 private:
-								FStoreClient() = default;
-								FStoreClient(const FStoreClient&) = delete;
-								FStoreClient(const FStoreClient&&) = delete;
-	void						operator = (const FStoreClient&) = delete;
-	void						operator = (const FStoreClient&&) = delete;
+											FStoreClient() = default;
+											FStoreClient(const FStoreClient&) = delete;
+											FStoreClient(const FStoreClient&&) = delete;
+	void									operator = (const FStoreClient&) = delete;
+	void									operator = (const FStoreClient&&) = delete;
 };
 
 } // namespace Trace
diff --git a/Engine/Source/Developer/TraceInsights/Private/Insights/Tests/InsightsTestUtils.cpp b/Engine/Source/Developer/TraceInsights/Private/Insights/Tests/InsightsTestUtils.cpp
--- a/Engine/Source/Developer/TraceInsights/Private/Insights/Tests/InsightsTestUtils.cpp
+++ b/Engine/Source/Developer/TraceInsights/Private/Insights/Tests/InsightsTestUtils.cpp
@@ -145,7 +145,7 @@ bool FInsightsTestUtils::FileContainsString(const FString& PathToFile, const FSt
 bool FInsightsTestUtils::IsUnrealTraceServerReady(const TCHAR* Host, int32 Port) const
 {
 	using namespace UE::Trace;
-	TSharedPtr<FStoreClient> StoreClient(FStoreClient::Connect(Host, Port));
+	TUniquePtr<FStoreClient> StoreClient = FStoreClient::Create(Host, Port);
 	if (!StoreClient.IsValid())
 	{
 		Test->AddInfo(TEXT("Cannot connect to UTS"));
@@ -325,12 +325,14 @@ void FInsightsTestUtils::ResetSession() const
 
 TSharedPtr<UE::Trace::FStoreClient> FInsightsTestUtils::CreateStoreClient(const TCHAR* Host, int32 Port) const
 {
-	TSharedPtr<UE::Trace::FStoreClient> StoreClient(UE::Trace::FStoreClient::Connect(Host, Port));
-	if (!StoreClient.IsValid())
+	TUniquePtr<UE::Trace::FStoreClient> StoreClient = UE::Trace::FStoreClient::Create(Host, Port);
+	TSharedPtr<UE::Trace::FStoreClient> SharedStoreClient = TSharedPtr<UE::Trace::FStoreClient>(StoreClient.Release());
+	if (!SharedStoreClient.IsValid())
 	{
 		Test->AddInfo(TEXT("The StoreClient shouldn't be null"));
 	}
-	return StoreClient;
+	
+	return SharedStoreClient;
 }
 
 ////////////////////////////////////////////////////////////////////////////////////////////////////
diff --git a/Engine/Source/Developer/TraceInsights/Private/Insights/TraceInsightsModule.cpp b/Engine/Source/Developer/TraceInsights/Private/Insights/TraceInsightsModule.cpp
--- a/Engine/Source/Developer/TraceInsights/Private/Insights/TraceInsightsModule.cpp
+++ b/Engine/Source/Developer/TraceInsights/Private/Insights/TraceInsightsModule.cpp
@@ -202,7 +202,7 @@ FString FTraceInsightsModule::GetDefaultStoreDir()
 {
 	using UE::Trace::FStoreClient;
 
-	TUniquePtr<FStoreClient> StoreClient(FStoreClient::Connect(TEXT("localhost")));
+	TUniquePtr<FStoreClient> StoreClient(FStoreClient::Create(TEXT("localhost")));
 
 	if (!StoreClient)
 	{
diff --git a/Engine/Source/Developer/TraceInsightsFrontend/Private/InsightsFrontend/StoreService/TraceServerControl.cpp b/Engine/Source/Developer/TraceInsightsFrontend/Private/InsightsFrontend/StoreService/TraceServerControl.cpp
--- a/Engine/Source/Developer/TraceInsightsFrontend/Private/InsightsFrontend/StoreService/TraceServerControl.cpp
+++ b/Engine/Source/Developer/TraceInsightsFrontend/Private/InsightsFrontend/StoreService/TraceServerControl.cpp
@@ -54,7 +54,7 @@ void FTraceServerControl::MakeMenu(FMenuBuilder& Builder)
 	// Create client
 	if (!Client)
 	{
-		Client.Reset(UE::Trace::FStoreClient::Connect(*Host, Port));
+		Client = (UE::Trace::FStoreClient::Create(*Host, Port));
 		if (Client)
 		{
 			ChangeState(EState::NotConnected, EState::Connected);
@@ -260,18 +260,14 @@ void FTraceServerControl::OnStart_Clicked()
 			{
 				for (uint32 Attempts = 0; Attempts < GStartConnectAttempts; ++Attempts)
 				{
-					UE::Trace::FStoreClient* NewClient = UE::Trace::FStoreClient::Connect(TEXT("127.0.0.1"));
 					if (bIsCancelRequested)
 					{
-						if (NewClient)
-						{
-							delete NewClient;
-						}
 						break;
 					}
-					if (NewClient)
+
+					Client = UE::Trace::FStoreClient::Create(TEXT("127.0.0.1"));
+					if (Client)
 					{
-						Client.Reset(NewClient);
 						break;
 					}
 					FPlatformProcess::Sleep(GStartConnectFrequencySeconds);
diff --git a/Engine/Source/Programs/TraceAnalyzer/Private/Common.h b/Engine/Source/Programs/TraceAnalyzer/Private/Common.h
--- a/Engine/Source/Programs/TraceAnalyzer/Private/Common.h
+++ b/Engine/Source/Programs/TraceAnalyzer/Private/Common.h
@@ -65,7 +65,7 @@ inline TUniquePtr<UE::Trace::IInDataStream> GetDataStream(int ArgC, TCHAR const*
 	}
 
 	uint32 StorePort = FCString::Strtoui64(ArgV[1], nullptr, 10);
-	if (UE::Trace::FStoreClient* StoreClient = UE::Trace::FStoreClient::Connect(TEXT("127.0.0.1"), StorePort))
+	if (TUniquePtr<UE::Trace::FStoreClient> StoreClient = UE::Trace::FStoreClient::Create(TEXT("127.0.0.1"), StorePort))
 	{
 		uint32 TraceId = FCString::Strtoui64(ArgV[2], nullptr, 10);
 		return StoreClient->ReadTrace(TraceId);
diff --git a/Engine/Source/Runtime/Experimental/TraceBasedDebuggersAnalysis/Private/TraceSessionsManager.cpp b/Engine/Source/Runtime/Experimental/TraceBasedDebuggersAnalysis/Private/TraceSessionsManager.cpp
--- a/Engine/Source/Runtime/Experimental/TraceBasedDebuggersAnalysis/Private/TraceSessionsManager.cpp
+++ b/Engine/Source/Runtime/Experimental/TraceBasedDebuggersAnalysis/Private/TraceSessionsManager.cpp
@@ -339,7 +339,7 @@ FString FTraceSessionsManager::ConnectToLiveSession_Relay(const FGuid RemoteSess
 
 FString FTraceSessionsManager::GetLocalTraceStoreDirPath()
 {
-	const Trace::FStoreClient* StoreClient = Trace::FStoreClient::Connect(TEXT("localhost"));
+	const TUniquePtr<Trace::FStoreClient> StoreClient = Trace::FStoreClient::Create(TEXT("localhost"));
 
 	if (!StoreClient)
 	{
@@ -407,7 +407,7 @@ const Trace::FStoreClient::FSessionInfo* FTraceSessionsManager::GetTraceSessionI
 	}
 
 	using namespace UE::Trace;
-	const FStoreClient* StoreClient = FStoreClient::Connect(InSessionHost.GetData());
+	const TUniquePtr<FStoreClient> StoreClient = FStoreClient::Create(InSessionHost.GetData());
 
 	if (!StoreClient)
 	{
@@ -427,7 +427,7 @@ FString FTraceSessionsManager::GetTraceFileNameFromStoreForSession(const FString
 		return FString();
 	}
 
-	const FStoreClient* StoreClient = FStoreClient::Connect(InSessionHost.GetData());
+	const TUniquePtr<FStoreClient> StoreClient = FStoreClient::Create(InSessionHost.GetData());
 	if (!StoreClient)
 	{
 		return FString();
@@ -472,7 +472,7 @@ FString FTraceSessionsManager::GetTraceFileNameFromStoreForSession(const FString
 Trace::FStoreClient::FTraceData FTraceSessionsManager::GetTraceDataStreamFromStore(const FStringView InSessionHost, const uint32 SessionID)
 {
 	using namespace UE::Trace;
-	const FStoreClient* StoreClient = FStoreClient::Connect(InSessionHost.GetData());
+	const TUniquePtr<FStoreClient> StoreClient = FStoreClient::Create(InSessionHost.GetData());
 	if (!StoreClient)
 	{
 		return FStoreClient::FTraceData();
