[Insights] Datastream and Insights manager support for secure connections.

- Added FSecureDataStream for secure Trace connections.
- Added API for creating analysis sessions with secure connections.
- Added requirement for Asio ssl support.
- Added log category for Trace data streams.

#Synced-CL 49918891

diff --git a/Engine/Source/Developer/TraceAnalysis/Private/DataStream.cpp b/Engine/Source/Developer/TraceAnalysis/Private/DataStream.cpp
--- a/Engine/Source/Developer/TraceAnalysis/Private/DataStream.cpp
+++ b/Engine/Source/Developer/TraceAnalysis/Private/DataStream.cpp
@@ -9,13 +9,15 @@
 #include "HAL/PlatformProcess.h"
 #include "HAL/Runnable.h"
 #include "HAL/RunnableThread.h"
-
+#include "Logging/LogCategory.h"
 #include "Sockets.h"
 
 #if PLATFORM_LINUX
 #include <sys/socket.h>
 #endif
 
+DEFINE_LOG_CATEGORY(LogTraceData);
+
 namespace UE {
 namespace Trace {
 
@@ -223,7 +225,7 @@ uint16 FDirectSocketStream::StartListening(uint16 Port)
 	// We cannot block on this thread or on the analysis thread (supports cancellation), create
 	// a short-lived thread for the blocking accept call.
 	ListeningThread = TUniquePtr<FRunnableThread>(FRunnableThread::Create(this, TEXT("FDirectSocketStreamListener")));
-	UE_LOG(LogCore, Log, TEXT("Started listening thread for direct trace connection on port %u"), Port);
+	UE_LOG(LogTraceData, Log, TEXT("Started listening thread for direct trace connection on port %u"), Port);
 
 	return Port;
 }
@@ -244,7 +246,7 @@ bool FDirectSocketStream::CreateSocket(uint16 Port)
 #endif
 	if (ErrorCode)
 	{
-		UE_LOG(LogCore, Error, TEXT("Failed to open acceptor: '%s'"), ANSI_TO_TCHAR(ErrorCode.message().c_str()));
+		UE_LOG(LogTraceData, Error, TEXT("Failed to open acceptor: '%s'"), ANSI_TO_TCHAR(ErrorCode.message().c_str()));
 		return false;
 	}
 
@@ -258,14 +260,14 @@ bool FDirectSocketStream::CreateSocket(uint16 Port)
 	Acceptor.bind(Endpoint, ErrorCode);
 	if (ErrorCode)
 	{
-		UE_LOG(LogCore, Error, TEXT("Failed to bind socket on port %u: '%s'"), Port, ANSI_TO_TCHAR(ErrorCode.message().c_str()));
+		UE_LOG(LogTraceData, Error, TEXT("Failed to bind socket on port %u: '%s'"), Port, ANSI_TO_TCHAR(ErrorCode.message().c_str()));
 		return false;
 	}
 
 	Acceptor.listen(MaxQueuedConnections, ErrorCode);
 	if (ErrorCode)
 	{
-		UE_LOG(LogCore, Error, TEXT("Failed to listen on port %u: '%s'"), Port, ANSI_TO_TCHAR(ErrorCode.message().c_str()));
+		UE_LOG(LogTraceData, Error, TEXT("Failed to listen on port %u: '%s'"), Port, ANSI_TO_TCHAR(ErrorCode.message().c_str()));
 		return false;
 	}
 
@@ -291,7 +293,7 @@ void FDirectSocketStream::Accept()
 		{
 			return;
 		}
-		UE_LOG(LogCore, Error, TEXT("Failed accept sockets connection, error: %u."), ErrorCode.value());
+		UE_LOG(LogTraceData, Error, TEXT("Failed accept sockets connection, error: %u."), ErrorCode.value());
 		return;
 	}
 
@@ -303,7 +305,7 @@ void FDirectSocketStream::Accept()
 		const auto Address = Socket.remote_endpoint().address();
 		TStringBuilder<16> RemoteAddress;
 		RemoteAddress << ANSI_TO_TCHAR(Address.to_string().c_str());
-		UE_LOG(LogCore, Log, TEXT("Accepted direct trace connection from %s"), RemoteAddress.ToString());
+		UE_LOG(LogTraceData, Log, TEXT("Accepted direct trace connection from %s"), RemoteAddress.ToString());
 	}
 
 	AsioContext->Acceptor.close();
@@ -325,7 +327,7 @@ bool FDirectSocketStream::BindToExistingConnection(FSocket&& ExistingConnectedSo
 
 	if (ErrorCode)
 	{
-		UE_LOG(LogCore, Error, TEXT("%hs: Failed to bind ASIO socket from FSocket connection, error: %u."), __func__, ErrorCode.value());
+		UE_LOG(LogTraceData, Error, TEXT("%hs: Failed to bind ASIO socket from FSocket connection, error: %u."), __func__, ErrorCode.value());
 		return false;
 	}
 
diff --git a/Engine/Source/Developer/TraceAnalysis/Private/DataStreamInternal.h b/Engine/Source/Developer/TraceAnalysis/Private/DataStreamInternal.h
--- a/Engine/Source/Developer/TraceAnalysis/Private/DataStreamInternal.h
+++ b/Engine/Source/Developer/TraceAnalysis/Private/DataStreamInternal.h
@@ -2,8 +2,11 @@
 #pragma once
 
 #include "Trace/DataStream.h"
+#include "Logging/LogCategory.h"
 #include "Asio.h"
 
+DECLARE_LOG_CATEGORY_EXTERN(LogTraceData, Log, Log);
+
 namespace UE {
 namespace Trace {
 	
@@ -24,4 +27,4 @@ private:
 
 	
 } // namespace Trace
-} // namespace UE
\ No newline at end of file
+} // namespace UE
diff --git a/Engine/Source/Developer/TraceAnalysis/Private/SecureDataStream.cpp b/Engine/Source/Developer/TraceAnalysis/Private/SecureDataStream.cpp
new file mode 100644
--- /dev/null
+++ b/Engine/Source/Developer/TraceAnalysis/Private/SecureDataStream.cpp
@@ -0,0 +1,586 @@
+// Copyright Epic Games, Inc. All Rights Reserved.
+
+#include "Trace/SecureDataStream.h"
+#include "Containers/AnsiString.h"
+#include "Containers/ArrayView.h"
+#include "Containers/UnrealString.h"
+#include "DataStreamInternal.h"
+#include "DataStreamInternal.h"
+#include "GenericPlatform/GenericPlatformFile.h"
+#include "HAL/Event.h"
+#include "HAL/PlatformFileManager.h"
+#include "HAL/PlatformProcess.h"
+#include "HAL/Runnable.h"
+#include "HAL/RunnableThread.h"
+#include "Sockets.h"
+#include "Templates/UniquePtr.h"
+
+#if WITH_SHARED_ASIO_SSL_SUPPORT
+#include "asio/ssl/stream.hpp"
+#include "asio/ssl/error.hpp"
+#include "asio/ssl/context.hpp"
+#include <openssl/ssl.h>
+#endif
+
+#if PLATFORM_LINUX
+#include <sys/socket.h>
+#endif
+
+// Enable extra debugging for secure connections
+#define UE_TRACE_ANALYSIS_SSS_DEBUG_ENABLE 0
+
+#if UE_TRACE_ANALYSIS_SSS_ENABLED 
+
+namespace UE::Trace {
+
+struct FSecureSocketContext
+{
+	FSecureSocketContext()
+		: Context(1)
+		, SSLContext(asio::ssl::context::sslv23_client)
+	{
+#if UE_TRACE_ANALYSIS_SSS_DEBUG_ENABLE // Enables additional SSL debug output
+		SSL_CTX_set_info_callback(SSLContext.native_handle(), [](const SSL* Ssl, int Type, int Value) {
+			{
+				UE_LOG(LogTraceData, VeryVerbose, 
+					TEXT("TLS Info [%hs]|Type [%hs]|Desc %d/[%hs]"), 
+					(Type & SSL_CB_READ) ? "Read" : "Write", 
+					SSL_alert_type_string_long(Type), 
+					Value,
+					SSL_alert_desc_string_long(Value)
+				);
+			}
+		});
+
+		SSL_CTX_set_msg_callback(SSLContext.native_handle(), [](int write_p, int version, int content_type, const void *buf, size_t len, SSL *ssl, void *arg) {
+			UE_LOG(LogTraceData, VeryVerbose, 
+				TEXT("TLS Traffic [%hs]|Type [%d]|Length [%zu]"), 
+				write_p ? "In" : "Out", 
+				content_type, 
+				len
+			);
+		});
+#endif
+	}
+
+	~FSecureSocketContext() 
+	{
+		Socket.Reset();	
+	}
+
+	bool OnVerifyCertificate(bool bPreverified, asio::ssl::verify_context &InContext)
+	{
+		char SubjectName[256];
+		X509 *Certificate = X509_STORE_CTX_get_current_cert(InContext.native_handle());
+		X509_NAME_oneline(X509_get_subject_name(Certificate), SubjectName, 256);
+		UE_LOG(LogTraceData, Verbose, TEXT("Verifying certificate '%hs'"), SubjectName)
+		return bPreverified;
+	}
+
+	asio::io_context Context;
+	asio::ssl::context SSLContext;
+	TUniquePtr<asio::ssl::stream<asio::ip::tcp::socket>> Socket;
+};
+
+FSecureSocketStream::FSecureSocketStream()
+	: ReadyToReadEvent(FPlatformProcess::GetSynchEventFromPool(false))
+	, State(EState::None)
+{
+	AsioContext = MakeUnique<FSecureSocketContext>();
+}
+
+FSecureSocketStream::FSecureSocketStream(TUniquePtr<FArchiveFileWriterGeneric>&& InFileWriter)
+	: ReadyToReadEvent(FPlatformProcess::GetSynchEventFromPool(false))
+	, FileWriter(MoveTemp(InFileWriter))
+	, State(EState::None)
+{
+	AsioContext = MakeUnique<FSecureSocketContext>();
+}
+
+FSecureSocketStream::~FSecureSocketStream()
+{
+	// Deregister our app_data value, if it was set with pre shared keys callback,
+	// otherwise asio thinks it's a verify callback base and try to delete it.
+	// Note that we check if the value is self, and only then reset it, otherwise we
+	// will leak the internal asio object.
+	SSL_CTX *NativeContext = AsioContext->SSLContext.native_handle();
+	if (Mode == ESecureMode::PreSharedKeys && NativeContext && SSL_CTX_get_app_data(NativeContext) == this)
+	{
+		SSL_CTX_set_app_data(NativeContext, nullptr);
+	}
+
+	if (AsioContext)
+	{
+		AsioContext->Context.stop();
+#if PLATFORM_LINUX
+		// Socket needs to be shutdown before closing the asio acceptor, otherwise ListeningThread would deadlock.
+		shutdown(AsioContext->Socket.native_handle(), SHUT_RD);
+#endif
+	}
+
+	if (IoThread)
+	{
+		IoThread->WaitForCompletion();
+		IoThread.Reset();
+	}
+
+	AsioContext.Reset();
+}
+
+
+bool FSecureSocketStream::LoadCertificateAuthority(FStringView InCertificatePath)
+{
+	SSL_CTX *NativeContextHandle = AsioContext->SSLContext.native_handle();
+	if (!NativeContextHandle) {
+		UE_LOG(LogTraceData, Error, TEXT("Failed to obtain the the native SSL context handle"));
+		return false;
+	}
+
+	FAnsiString CertificatePath(InCertificatePath);
+	if (!SSL_CTX_load_verify_locations(NativeContextHandle, *CertificatePath, nullptr))
+	{
+		UE_LOG(LogTraceData, Error, TEXT("Failed to load certificate authority"));
+		return false;
+	}
+
+	return true;
+}
+
+bool FSecureSocketStream::SetCertificateAuthority(TArrayView<uint8> CA)
+{
+	const unsigned char* CAPointer = CA.GetData(); 
+	size_t CASize = CA.Num();
+	X509* DecodedCert = d2i_X509(NULL, &CAPointer, CASize);
+
+	if (!DecodedCert)
+	{
+		UE_LOG(LogTraceData, Error, TEXT("Failed to decode SSL certificate authority"));
+		return false;
+	}
+
+	SSL_CTX *NativeContextHandle = AsioContext->SSLContext.native_handle();
+	if (!NativeContextHandle) {
+		UE_LOG(LogTraceData, Error, TEXT("Failed to obtain the the native SSL context handle"));
+		return false;
+	}
+
+	X509_STORE *CertStore = SSL_CTX_get_cert_store(NativeContextHandle);
+	if (!CertStore) {
+		UE_LOG(LogTraceData, Error, TEXT("Failed to obtain the cert store "));
+		return false;
+	}
+	if (!X509_STORE_add_cert(CertStore, DecodedCert)) {
+		UE_LOG(LogTraceData, Error, TEXT("Failed to add certificate authority "));
+		return false;
+	}
+
+	UE_LOG(LogTraceData, Log, TEXT("Succesfully applied Certificate Authority to SSL context"));
+	return true;
+}
+
+struct FSecureStreamPSKCallback
+{
+	static unsigned int OnPSK(SSL* Ssl, const char* Hint, char* Identity, unsigned int MaxIdentityLength, unsigned char* PSK, unsigned int MaxPSKLength)
+	{
+		SSL_CTX* Context = SSL_get_SSL_CTX(Ssl);
+		FSecureSocketStream* Self = (FSecureSocketStream*) SSL_CTX_get_app_data(Context);	
+		if (!Self)
+		{
+			UE_LOG(LogTraceData, Error, TEXT("Failed to retrieve app data"));
+			return 0;
+		}
+
+		FAnsiStringView PSKIdentity(Self->Identity);
+		const int32 IdentityCount = PSKIdentity.CopyString(Identity, MaxIdentityLength);
+		Identity[IdentityCount] = '\0';
+		UE_LOG(LogTraceData, Verbose, TEXT("Using identity '%hs' (hint: '%hs')"), Identity, Hint);
+
+		if (MaxPSKLength < uint32(Self->AuthToken.Num()))
+		{
+			UE_LOG(LogTraceData, Error, TEXT("Target PSK buffer too small."));
+			return 0;
+		}
+
+		const uint32 SizeToCopy = Self->AuthToken.Num();
+		FMemory::Memcpy(PSK, Self->AuthToken.GetData(), SizeToCopy);
+
+#if UE_TRACE_ANALYSIS_SSS_DEBUG_ENABLE
+		auto ByteToHex = [](uint8 *Data, uint32 Size, char *OutBuffer)
+		{
+			for (uint32 i = 0; i < Size; ++i)
+			{
+				int written = snprintf(OutBuffer, 3, "%02X", static_cast<int>(Data[i]));
+				if (written < 2 || written > 2)
+				{
+					*OutBuffer = '\0';
+					break;
+				}
+				OutBuffer += 2;
+			}
+			*OutBuffer = '\0';
+		};
+
+		char Hex[256];
+		ByteToHex(PSK, SizeToCopy, Hex);
+		UE_LOG(LogTraceData, Verbose, TEXT("PSK: '%hs'"), Hex);
+#endif // UE_TRACE_ANALYSIS_SSS_DEBUG_ENABLE
+
+		return SizeToCopy;
+	}
+};
+
+bool FSecureSocketStream::SetPreSharedKey()
+{
+	AsioContext->SSLContext.clear_options(asio::ssl::context::default_workarounds);
+
+	SSL_CTX *NativeContextHandle = AsioContext->SSLContext.native_handle();
+	if (!NativeContextHandle) {
+		UE_LOG(LogTraceData, Error, TEXT("Failed to obtain the native SSL context handle"));
+		return false;
+	}
+
+	if (!SSL_CTX_set_app_data(NativeContextHandle, this))
+	{
+		UE_LOG(LogTraceData, Error, TEXT("Failed to set app data"));
+		return false;
+	}
+
+	SSL_CTX_set_min_proto_version(NativeContextHandle, TLS1_2_VERSION);
+	SSL_CTX_set_psk_client_callback(NativeContextHandle, &FSecureStreamPSKCallback::OnPSK);
+
+	static const char* CipherList = "ECDHE-PSK-CHACHA20-POLY1305";
+	if (!SSL_CTX_set_cipher_list(NativeContextHandle, CipherList))
+	{
+		UE_LOG(LogTraceData, Error, TEXT("Failed to set cipher list"));
+		return false;
+	}
+
+	static const char* GroupLists = "P-384:P-256";
+	if (!SSL_CTX_set1_groups_list(NativeContextHandle, GroupLists))
+	{
+		UE_LOG(LogTraceData, Error, TEXT("Failed to set group list"));
+		return false;
+	}
+
+	UE_LOG(LogTraceData, VeryVerbose, TEXT("Cipher list: '%hs', groups list: '%hs'"), CipherList, GroupLists);
+
+	return true;
+}
+
+
+void FSecureSocketStream::SetState(EState NewState, bool bTriggerReady)
+{
+	if (State != NewState)
+	{
+#if UE_TRACE_ANALYSIS_SSS_DEBUG_ENABLE 
+		auto LexStateToString = [](EState State) -> const TCHAR* {
+			switch(State)
+			{
+			case EState::None: return TEXT("None");
+			case EState::Closed: return TEXT("Closed");
+			case EState::Connecting: return TEXT("Connecting");
+			case EState::Handshaking: return TEXT("Handshaking");
+			case EState::Authenticating: return TEXT("Authenticating");
+			case EState::Open: return TEXT("Open");
+			case EState::Failed: return TEXT("Failed");
+			default: return TEXT("Unknown");
+			}
+		};
+		UE_LOG(LogTraceData, Verbose, TEXT("Secure connection state: %s -> %s"), LexStateToString(State), LexStateToString(NewState));
+#endif //UE_DEBUG
+		State = NewState;
+	}
+	if (bTriggerReady)
+	{
+		ReadyToReadEvent->Trigger();
+	}
+}
+
+bool FSecureSocketStream::ConnectWithCertificate(FStringView InHost, uint16 InPort, TArrayView<uint8> InAuthToken, FStringView InCertificatePath)
+{
+	if (InCertificatePath.IsEmpty() || !LoadCertificateAuthority(InCertificatePath))
+	{
+		return false;
+	}
+	Mode = ESecureMode::Certificates;
+	return ConnectInternal(InHost, InPort, InAuthToken);
+}
+
+bool FSecureSocketStream::ConnectWithSelfSigned(FStringView InHost, uint16 InPort, TArrayView<uint8> InAuthToken, TArrayView<uint8> InCertificate)
+{
+	if (!SetCertificateAuthority(InCertificate))
+	{
+		return false;
+	}
+	
+	Mode = ESecureMode::Certificates;
+	return ConnectInternal(InHost, InPort, InAuthToken);
+}
+
+bool FSecureSocketStream::ConnectWithPreSharedKey(FStringView InHost, uint16 InPort, TArrayView<uint8> InPreSharedKey, FAnsiStringView InIdentity)
+{
+	if (!SetPreSharedKey())
+	{
+		return false;
+	}
+	Identity = InIdentity;
+	Mode = ESecureMode::PreSharedKeys;
+	return ConnectInternal(InHost, InPort, InPreSharedKey);
+}
+
+bool FSecureSocketStream::ConnectInternal(FStringView InHost, uint16 InPort, TArrayView<uint8> InAuthToken)
+{
+	// Create the socket object
+	AsioContext->Socket = MakeUnique<asio::ssl::stream<asio::ip::tcp::socket>>(AsioContext->Context, AsioContext->SSLContext);
+	check(Mode != ESecureMode::None);
+	if (Mode == ESecureMode::Certificates)
+	{
+		AsioContext->Socket->set_verify_mode(asio::ssl::verify_peer);
+		AsioContext->Socket->set_verify_callback(
+			std::bind(
+				&FSecureSocketContext::OnVerifyCertificate, 
+				AsioContext.Get(), 
+				std::placeholders::_1, 
+				std::placeholders::_2
+			)
+		);
+	}
+	else
+	{
+		AsioContext->Socket->set_verify_mode(asio::ssl::verify_none);
+	}
+	
+	//Convert host to ansistring
+	TAnsiStringBuilder<128> AnsiHost;
+	AnsiHost << InHost;
+	TAnsiStringBuilder<10> AnsiPort;
+	AnsiPort << InPort;
+	AuthToken = InAuthToken;
+
+	asio::ip::tcp::resolver Resolver(AsioContext->Context);
+	auto Endpoints = Resolver.resolve(AnsiHost.ToString(), AnsiPort.ToString());
+
+	State = EState::Connecting;
+
+	asio::async_connect(
+		AsioContext->Socket->lowest_layer(), 
+		Endpoints,
+		[this](const std::error_code &Error, const asio::ip::tcp::endpoint& Endpoint) {
+			if (Error)
+			{
+				UE_LOG(LogTraceData, Error, 
+					TEXT("Failed to connect to tracing host '%hs' on port %u: %hs"), 
+					Endpoint.address().to_string().c_str(),
+					Endpoint.port(),
+					Error.message().c_str()
+				);
+				SetState(EState::Failed, true);
+				return;
+			}
+			UE_LOG(LogTraceData, Display, 
+				TEXT("Secure connection opened to '%hs' on port %u."),
+				Endpoint.address().to_string().c_str(),
+				Endpoint.port()
+			);
+			AsyncHandshake();
+		}
+	);
+
+	IoThread = TUniquePtr<FRunnableThread>(FRunnableThread::Create(this, TEXT("FSecureConnectionStream")));
+	return true;
+}
+
+void FSecureSocketStream::AsyncHandshake()
+{
+	State = EState::Handshaking;
+
+	AsioContext->Socket->async_handshake(
+		asio::ssl::stream_base::client, 
+		[this](const std::error_code &Error) 
+		{
+			if (Error)
+			{
+				UE_LOG(LogTraceData, Error, 
+					TEXT("Failed secure handshake: %hs"), 
+					Error.message().c_str()
+				);
+				SetState(EState::Failed, true);
+				return;
+			}
+			UE_LOG(LogTraceData, Verbose, TEXT("Secure handshake succeeded"));
+			AsyncAuthenticate();
+		}
+	);
+}
+
+void FSecureSocketStream::AsyncAuthenticate()
+{
+	State = EState::Authenticating;
+	AsioContext->Socket->async_write_some(
+		asio::buffer(AuthToken.GetData(), AuthToken.Num()),
+		[this] (asio::error_code Error, size_t BytesWritten)
+		{
+			if (Error)
+			{
+				UE_LOG(LogTraceData, Error, 
+					TEXT("Failed to write authorization token: %hs"), 
+					Error.message().c_str()
+				);
+				SetState(EState::Failed, true);
+				return;
+			}
+			UE_LOG(LogTraceData, Verbose, TEXT("Authentication sent"));
+			SetState(EState::Open, true);
+			ReadyToReadEvent->Trigger();
+#if UE_TRACE_ANALYSIS_SSS_BUFFERED
+			AsyncRecieveData();
+#endif
+		}
+	);
+}
+
+#if UE_TRACE_ANALYSIS_SSS_BUFFERED
+void FSecureSocketStream::AsyncRecieveData()
+{
+	auto Handler = [this](const asio::error_code Error, size_t BytesRead)
+	{
+		if (Error == asio::error::eof || Error == asio::ssl::error::stream_truncated)
+		{
+			UE_LOG(LogTraceData, Log, 
+				TEXT("Connection closed by tracing peer")
+			);
+			State = EState::Closed;
+			return;
+		}
+		else if (Error)
+		{
+			UE_LOG(LogTraceData, Error, 
+				TEXT("Read error recieving data: %hs"),
+				Error.message().c_str()
+			);
+			State = EState::Failed;
+			return;
+		}
+		
+		if (BytesRead)
+		{
+			// Accumulate small reads into ReadBuffer
+			BufferedBytes += BytesRead;
+
+			// Flush buffered data to queue if the reader is waiting
+			// or read buffer is full.
+			if (ReadQueue.IsEmpty() || BufferedBytes == ReadBuffer.Num())
+			{
+				//UE_LOG(LogTraceData, VeryVerbose, TEXT("[IO] Flushing %d bytes to queue"), BufferedBytes);
+				ReadQueue.Enqueue(ReadBuffer.GetData(), BufferedBytes);
+				BufferedBytes = 0u;
+			}
+		}
+
+		AsyncRecieveData();
+	};
+
+	if (!ReadBuffer.Num())
+	{
+		ReadBuffer.AddUninitialized(BufferCapacity);
+	}
+	uint8* ReadBufferPtr = ReadBuffer.GetData() + BufferedBytes;
+	int32 BytesToRead = ReadBuffer.Num() - BufferedBytes; 
+	checkf(BytesToRead > 0, TEXT("BytesToRead: %d, BufferSize: %d"), BytesToRead, ReadBuffer.Num());
+	AsioContext->Socket->async_read_some(asio::buffer(ReadBufferPtr, BytesToRead), Handler);
+}
+#endif
+
+uint32 FSecureSocketStream::Run()
+{
+	UE_LOG(LogTraceData, Verbose, TEXT("Starting IO thread for secure connection stream"));
+	AsioContext->Context.run();
+	UE_LOG(LogTraceData, Verbose, TEXT("Finishing IO thread for secure connection stream"));
+	return 0;
+}
+
+void FSecureSocketStream::Stop()
+{
+	AsioContext->Context.stop();
+}
+
+bool FSecureSocketStream::WaitUntilReady()
+{
+	return ReadyToReadEvent->Wait(0, true);
+}
+
+int32 FSecureSocketStream::Read(void* OutData, uint32 Size)
+{
+	auto MaybeWriteFile = [this](void* Data, uint32 Size)
+	{
+		if (Size > 0 && FileWriter.IsValid() /*&& !FileWriter->IsError()*/)
+		{
+			FileWriter->Serialize(Data, Size);
+		}
+		return Size;
+	};
+
+	const int32 OutSize = int32(Size);
+	int32 BytesRead = 0;
+#if UE_TRACE_ANALYSIS_SSS_BUFFERED
+	uint32 DataBlocksRead = 0;
+	do {
+		while (FData* NextData = ReadQueue.Peek())
+		{
+			// Can we fit next data block into the out buffer?
+			if ((NextData->Num() + BytesRead) > OutSize)
+			{
+				return MaybeWriteFile(OutData, BytesRead);
+			}
+			// Peel off the block
+			TOptional<FData> DequedData = ReadQueue.Dequeue();
+			if (!DequedData)
+			{
+				return MaybeWriteFile(OutData, BytesRead);
+			}
+			FData& Data = DequedData.GetValue();
+			FMemory::Memcpy((uint8*)OutData + BytesRead, Data.GetData(), Data.Num());
+			BytesRead += Data.Num();
+			if (++DataBlocksRead >= 10)
+			{
+				return MaybeWriteFile(OutData, BytesRead);
+			}
+		}
+	}
+	while (State == EState::Open && BytesRead < OutSize);
+#else
+	asio::error_code Error;
+	BytesRead = AsioContext->Socket->read_some(asio::buffer(OutData, Size), Error);
+	if (Error)
+	{
+		Close();
+		if (Error == asio::error::eof || Error == asio::ssl::error::stream_truncated || Error == asio::error::connection_reset)
+		{
+			UE_LOG(LogTraceData, Log, 
+				TEXT("Connection closed by tracing peer")
+			);
+			return 0;
+		}
+
+		UE_LOG(LogTraceData, Error, 
+			TEXT("Read error recieving data (%d): %hs"),
+			Error.value(),
+			Error.message().c_str()
+		);
+		return Error.value() > 0 ? -Error.value() : Error.value();
+	}
+	check(BytesRead > 0); // If no bytes read an error should have been set.
+#endif
+
+	return MaybeWriteFile(OutData, BytesRead);
+}
+
+void FSecureSocketStream::Close()
+{
+	State = EState::Closed;
+	FileWriter.Reset();
+}
+
+} // namespace UE::Trace
+
+#endif // UE_TRACE_ANALYSIS_SSS_ENABLED 
diff --git a/Engine/Source/Developer/TraceAnalysis/Public/Trace/SecureDataStream.h b/Engine/Source/Developer/TraceAnalysis/Public/Trace/SecureDataStream.h
new file mode 100644
--- /dev/null
+++ b/Engine/Source/Developer/TraceAnalysis/Public/Trace/SecureDataStream.h
@@ -0,0 +1,150 @@
+// Copyright Epic Games, Inc. All Rights Reserved.
+#pragma once
+
+#include "DataStream.h"
+#include "Containers/AnsiString.h"
+#include "CoreFwd.h"
+#include "HAL/FileManagerGeneric.h"
+#include "HAL/Platform.h"
+#include "HAL/Runnable.h"
+#include "HAL/Runnable.h"
+#include "Containers/SpscQueue.h"
+#include "Templates/UniquePtr.h"
+
+#define UE_API TRACEANALYSIS_API
+
+// Secure socket stream support depends on SSL support in ASIO
+#define UE_TRACE_ANALYSIS_SSS_ENABLED WITH_SHARED_ASIO_SSL_SUPPORT
+
+// Enables experimental buffered receiving path
+// for FSecureSocketStream
+#define UE_TRACE_ANALYSIS_SSS_BUFFERED 1
+
+#if UE_TRACE_ANALYSIS_SSS_ENABLED 
+
+namespace UE::Trace
+{
+
+struct FSecureStreamPSKCallback;
+
+class FSecureSocketStream : public UE::Trace::IInDataStream, public ::FRunnable
+{
+public:
+	UE_API FSecureSocketStream();
+	UE_API FSecureSocketStream(TUniquePtr<FArchiveFileWriterGeneric>&& InFileWriter);
+	UE_API virtual ~FSecureSocketStream() override;
+
+	/**
+	 * Connect to a tracing host. Specify a authorization token and an optional trusted CA certificate file.
+	 * @param Host Hostname to connect to.
+	 * @param Port Port to connect to.
+	 * @param AuthToken Authorization token to send.
+	 * @param CertificatePath Optional path to trusted root certificate
+	 * @return True on success, false if not.
+	 */
+	UE_API bool ConnectWithCertificate(FStringView Host, uint16 Port, TArrayView<uint8> AuthToken, FStringView CertificatePath);
+
+	/**
+	 * Connect to a tracing host. Specify a authorization token and an optional trusted in-memory CA certificate.
+	 * @param Host Hostname to connect to.
+	 * @param Port Port to connect to.
+	 * @param AuthToken Authorization token to send.
+	 * @param Certificate A self-signed root certificate in DER format.
+	 * @return True on success, false if not.
+	 */
+	UE_API bool ConnectWithSelfSigned(FStringView Host, uint16 Port, TArrayView<uint8> AuthToken, TArrayView<uint8> Certificate);
+
+	/**
+	 * Connect to a tracing host. Specify an authorization token used for pre-shared key verification.
+	 * @param Host Hostname to connect to.
+	 * @param Port Port to connect to.
+	 * @param PreSharedKey Key to use in secure handshake
+	 * @param Identity Identity to use in secure handshake
+	 * @return True on success, false if not.
+	 */
+	UE_API bool ConnectWithPreSharedKey(FStringView Host, uint16 Port, TArrayView<uint8> PreSharedKey, FAnsiStringView Identity);
+
+
+	UE_API virtual int32 Read(void* Data, uint32 Size) override;
+
+private:
+	enum
+	{
+		DefaultPort = 1986,			// Default port to use.
+		MaxPortAttempts = 16,		// How many port increments are tried if fail to bind default.
+		MaxQueuedConnections = 4,	// Size of connection queue.
+	};
+
+	enum class EState : uint8 
+	{
+		None = 0,
+		Closed,
+		Connecting,
+		Handshaking,
+		Authenticating,
+		Open,
+		Failed = 0xff
+	};
+
+	bool ConnectInternal(FStringView Host, uint16 Port, TArrayView<uint8> AuthToken);
+
+	// IInStream interface
+	UE_API virtual bool WaitUntilReady() override;
+	UE_API virtual void Close() override;
+
+	// FRunnable interface
+	UE_API virtual uint32 Run() override;
+	UE_API virtual void Stop() override;
+
+	// Async data functions
+	void AsyncHandshake();
+	void AsyncAuthenticate();
+#if UE_TRACE_ANALYSIS_SSS_BUFFERED
+	void AsyncRecieveData();
+#endif
+
+	bool LoadCertificateAuthority(FStringView InCertificatePath);
+	bool SetCertificateAuthority(TArrayView<uint8> CA);
+	bool SetPreSharedKey();
+	void SetState(EState NewState, bool bTriggerReady);
+
+private:
+
+	enum class ESecureMode : uint8 
+	{
+		None,
+		Certificates,
+		PreSharedKeys
+	};
+	
+	class FData : public TArray<uint8>
+	{
+	public:
+		FData(uint8* Data, int32 Size)
+			: TArray<uint8>(TArrayView<uint8>(Data, Size))
+		{}
+	};
+
+	TUniquePtr<struct FSecureSocketContext> AsioContext;
+	TUniquePtr<class FRunnableThread> IoThread;
+	FEvent* ReadyToReadEvent;
+	TUniquePtr<class FArchiveFileWriterGeneric> FileWriter;
+	EState State;
+	TArray<uint8> AuthToken;
+	FAnsiString Identity;
+	ESecureMode Mode = ESecureMode::None;
+#if UE_TRACE_ANALYSIS_SSS_BUFFERED
+	static constexpr int32 BufferCapacity = 512 * 1024;
+	TArray<uint8> ReadBuffer;
+	uint32 BufferedBytes = 0u;
+	TSpscQueue<FData> ReadQueue;
+#endif
+
+	friend FSecureStreamPSKCallback;
+};
+
+} // namespace UE::Trace
+
+#endif // UE_TRACE_ANALYSIS_SSS_ENABLED 
+
+#undef UE_API
diff --git a/Engine/Source/Developer/TraceAnalysis/TraceAnalysis.Build.cs b/Engine/Source/Developer/TraceAnalysis/TraceAnalysis.Build.cs
--- a/Engine/Source/Developer/TraceAnalysis/TraceAnalysis.Build.cs
+++ b/Engine/Source/Developer/TraceAnalysis/TraceAnalysis.Build.cs
@@ -13,9 +13,19 @@ public class TraceAnalysis : ModuleRules
 				"Asio",
 				"Cbor",
 				"Core",
+				"OpenSSL",
 				"Sockets",
 				"TraceLog",
 			}
 		);
+
+        if (Target.bCompileAsioSSLSupport)
+        {
+	        PublicDefinitions.Add("WITH_SHARED_ASIO_SSL_SUPPORT=1");
+        }
+        else
+        {
+	        PublicDefinitions.Add("WITH_SHARED_ASIO_SSL_SUPPORT=0");
+        }
 	}
 }
diff --git a/Engine/Source/Developer/TraceInsights/Private/Insights/InsightsManager.cpp b/Engine/Source/Developer/TraceInsights/Private/Insights/InsightsManager.cpp
--- a/Engine/Source/Developer/TraceInsights/Private/Insights/InsightsManager.cpp
+++ b/Engine/Source/Developer/TraceInsights/Private/Insights/InsightsManager.cpp
@@ -17,6 +17,8 @@
 #include "Misc/Paths.h"
 #include "Modules/ModuleManager.h"
 #include "Templates/UniquePtr.h"
+#include "Trace/DataStream.h"
+#include "Trace/SecureDataStream.h"
 #include "WorkspaceMenuStructure.h"
 #include "WorkspaceMenuStructureModule.h"
 
@@ -875,6 +877,53 @@ uint16 FInsightsManager::ListenForDirectTrace(const TCHAR* InStreamName, uint16
 
 ////////////////////////////////////////////////////////////////////////////////////////////////////
 
+void FInsightsManager::ConnectToSecureTrace(FStringView Host, uint16 Port, TArrayView<uint8> AuthorizationToken, FStringView CertificateFile, TArrayView<uint8> CertificateAuthority, FAnsiStringView Identity)
+{
+#if UE_TRACE_ANALYSIS_SSS_ENABLED
+	TUniquePtr<UE::Trace::FSecureSocketStream> TraceData = MakeUnique<UE::Trace::FSecureSocketStream>();
+	bool bConnected = false;
+	if (!CertificateFile.IsEmpty())
+	{
+		bConnected = TraceData->ConnectWithCertificate(Host, Port, AuthorizationToken, CertificateFile);
+	}
+	else if (!CertificateAuthority.IsEmpty())
+	{
+		bConnected = TraceData->ConnectWithSelfSigned(Host, Port, AuthorizationToken, CertificateAuthority);
+	}
+	else
+	{
+		bConnected = TraceData->ConnectWithPreSharedKey(Host, Port, AuthorizationToken, Identity);
+	}
+
+	if (!bConnected)
+	{
+		UE_LOG(LogCore, Error, TEXT("Failed to connect to secure trace session."));
+	}
+
+	constexpr uint32 TraceId = 0;
+	TStringBuilder<64> TraceName;
+	TraceName << TEXT("Secure Trace Stream from ") << Host;
+
+	Session = AnalysisService->StartAnalysis(TraceId, TraceName.ToString(), MoveTemp(TraceData));
+
+	if (Session)
+	{
+		CurrentTraceStreamType = ETraceStreamType::DirectTrace;
+		CurrentTraceName = TraceName;
+		CurrentTraceId = 0;
+		CurrentDirectTracePort = Port;
+		bIsSessionInfoSet = false;
+		OnSessionChanged();
+	}
+	else if (bAutoQuit)
+	{
+		RequestEngineExit(AutoQuitMsgOnFail);
+	}
+#endif //UE_TRACE_ANALYSIS_SSS_ENABLED
+}
+
+////////////////////////////////////////////////////////////////////////////////////////////////////
+
 void FInsightsManager::StartAnalysisWithStream(TUniquePtr<UE::Trace::IInDataStream>&& InStream, const TCHAR* InStreamName)
 {
 	ResetSession();
diff --git a/Engine/Source/Developer/TraceInsights/Private/Insights/InsightsManager.h b/Engine/Source/Developer/TraceInsights/Private/Insights/InsightsManager.h
--- a/Engine/Source/Developer/TraceInsights/Private/Insights/InsightsManager.h
+++ b/Engine/Source/Developer/TraceInsights/Private/Insights/InsightsManager.h
@@ -2,6 +2,7 @@
 
 #pragma once
 
+#include "Containers/StringFwd.h"
 #include "CoreTypes.h"
 
 #include "Containers/Ticker.h"
@@ -257,6 +258,18 @@ public:
 	 */
 	uint16 ListenForDirectTrace(const TCHAR* InStreamName = nullptr, uint16 InPort = 0);
 
+	/**
+	 * Creates a new analysis session instance using a secure connection to the specified host.
+	 * @param Host Hostname to connect to.
+	 * @param Port Port to connect to.
+	 * @param AuthorizationToken Optional authorization token to use when connecting.
+	 * @param CertificateFile Optional certificate to use when connecting. Excludes usage of CertificateAuthority.
+	 * @param CertificateAuthority Optional certificate to use as trusted. CertificateFile must be nullptr.
+	 * @param Identity Identity string used to connect when using pre-shared keys.
+	 * @note All view-type parameters are copied and does not have to outlive this function call.
+	 */
+	void ConnectToSecureTrace(FStringView Host, uint16 Port, TArrayView<uint8> AuthorizationToken, FStringView CertificateFile = FStringView(), TArrayView<uint8> CertificateAuthority = TArrayView<uint8>(), FAnsiStringView Identity = FAnsiStringView());
+
 	/**
 	 * Creates a new analysis session instance from a provided data stream.
 	 * Replaces the current analysis session.
diff --git a/Engine/Source/Developer/TraceInsights/Private/Insights/TraceInsightsModule.cpp b/Engine/Source/Developer/TraceInsights/Private/Insights/TraceInsightsModule.cpp
--- a/Engine/Source/Developer/TraceInsights/Private/Insights/TraceInsightsModule.cpp
+++ b/Engine/Source/Developer/TraceInsights/Private/Insights/TraceInsightsModule.cpp
@@ -526,6 +526,13 @@ uint16 FTraceInsightsModule::StartAnalysisWithDirectTrace(const TCHAR* InStreamN
 	return Port;
 }
 
+////////////////////////////////////////////////////////////////////////////////////////////////////
+void FTraceInsightsModule::StartAnalysisWithSecureTrace(FStringView Host, uint16 Port, TArrayView<uint8> AuthorizationToken, FStringView CertificateFile, TArrayView<uint8> CertificateAuthority, FAnsiStringView Identity)
+{
+	FInsightsManager::Get()->ConnectToSecureTrace(Host, Port, AuthorizationToken, CertificateFile, CertificateAuthority, Identity);
+	UpdateAppTitle();
+}
+
 ////////////////////////////////////////////////////////////////////////////////////////////////////
 
 void FTraceInsightsModule::StartAnalysisWithStream(TUniquePtr<UE::Trace::IInDataStream>&& InStream, const TCHAR* InStreamName)
diff --git a/Engine/Source/Developer/TraceInsights/Private/Insights/TraceInsightsModule.h b/Engine/Source/Developer/TraceInsights/Private/Insights/TraceInsightsModule.h
--- a/Engine/Source/Developer/TraceInsights/Private/Insights/TraceInsightsModule.h
+++ b/Engine/Source/Developer/TraceInsights/Private/Insights/TraceInsightsModule.h
@@ -65,6 +65,7 @@ public:
 	virtual void StartAnalysisForLastLiveSession(float InRetryTime = 1.0f) override;
 	virtual void StartAnalysisForTraceFile(const TCHAR* InTraceFile) override;
 	virtual uint16 StartAnalysisWithDirectTrace(const TCHAR* InStreamName = nullptr, uint16 InPort = 0) override;
+	virtual void StartAnalysisWithSecureTrace(FStringView Host, uint16 Port, TArrayView<uint8> AuthorizationToken, FStringView CertificateFile = FStringView(), TArrayView<uint8> CertificateAuthority = TArrayView<uint8>(), FAnsiStringView Identity = FAnsiStringView()) override;
 	virtual void StartAnalysisWithStream(TUniquePtr<UE::Trace::IInDataStream>&& InStream, const TCHAR* InStreamName = nullptr) override;
 
 	virtual void ShutdownUserInterface() override;
diff --git a/Engine/Source/Developer/TraceInsights/Public/Insights/IUnrealInsightsModule.h b/Engine/Source/Developer/TraceInsights/Public/Insights/IUnrealInsightsModule.h
--- a/Engine/Source/Developer/TraceInsights/Public/Insights/IUnrealInsightsModule.h
+++ b/Engine/Source/Developer/TraceInsights/Public/Insights/IUnrealInsightsModule.h
@@ -308,6 +308,18 @@ public:
 	 */
 	virtual uint16 StartAnalysisWithDirectTrace(const TCHAR* InStreamName = nullptr, uint16 InPort = 0) = 0;
 
+	/**
+	 * Starts analysis for a secure trace trace connection.
+	 * @param Host Host to connect to.
+	 * @param Port Port to connect to.
+	 * @param AuthorizationToken Optional authorization token to use when connecting.
+	 * @param CertificateFile Optional certificate to use when connecting. Excludes usage of CertificateAuthority.
+	 * @param CertificateAuthority Optional certificate to use as trusted. CertificateFile must be nullptr.
+	 * @param Identity Identity string used to connect when using pre-shared keys.
+	 * @note All view-type parameters are copied and does not have to outlive this function call.
+	 */
+	virtual void StartAnalysisWithSecureTrace(FStringView Host, uint16 Port, TArrayView<uint8> AuthorizationToken, FStringView CertificateFile = FStringView(), TArrayView<uint8> CertificateAuthority = TArrayView<uint8>(), FAnsiStringView Identity = FAnsiStringView()) = 0;
+
 	/**
 	 * Start analysis given a data stream.
 	 *
diff --git a/Engine/Source/Programs/UnrealInsights/UnrealInsights.Target.cs b/Engine/Source/Programs/UnrealInsights/UnrealInsights.Target.cs
--- a/Engine/Source/Programs/UnrealInsights/UnrealInsights.Target.cs
+++ b/Engine/Source/Programs/UnrealInsights/UnrealInsights.Target.cs
@@ -32,6 +32,7 @@ public class UnrealInsightsTarget : TargetRules
 		bForceBuildTargetPlatforms = true;
 		bCompileWithStatsWithoutEngine = true;
 		bCompileWithPluginSupport = true;
+		bCompileAsioSSLSupport = true;
 
 		// For source code editor access & regex (crossplatform)
 		bIncludePluginsForTargetPlatforms = true;
