[Insights] Fix unix implementation

Removed code copied from the direct socket stream which needs to shut down the acceptor on Linux. We do not need the same a socket.

#Synced-CL 49918891

diff --git a/Engine/Source/Developer/TraceAnalysis/Private/SecureDataStream.cpp b/Engine/Source/Developer/TraceAnalysis/Private/SecureDataStream.cpp
--- a/Engine/Source/Developer/TraceAnalysis/Private/SecureDataStream.cpp
+++ b/Engine/Source/Developer/TraceAnalysis/Private/SecureDataStream.cpp
@@ -112,10 +112,6 @@ FSecureSocketStream::~FSecureSocketStream()
 	if (AsioContext)
 	{
 		AsioContext->Context.stop();
-#if PLATFORM_LINUX
-		// Socket needs to be shutdown before closing the asio acceptor, otherwise ListeningThread would deadlock.
-		shutdown(AsioContext->Socket.native_handle(), SHUT_RD);
-#endif
 	}
 
 	if (IoThread)
