[Insights] Change secure data types to const

* Changes array view to const to match input data.
* Fixed memory leak of X509 handles.
* Fixed transfer of authorization code in case it cannot be sent in one go.
* InsightsManager should call ResetSession when a secure session is created. Also support auto quit.

#Synced-CL 50170004

diff --git a/Engine/Source/Developer/TraceAnalysis/Private/SecureDataStream.cpp b/Engine/Source/Developer/TraceAnalysis/Private/SecureDataStream.cpp
--- a/Engine/Source/Developer/TraceAnalysis/Private/SecureDataStream.cpp
+++ b/Engine/Source/Developer/TraceAnalysis/Private/SecureDataStream.cpp
@@ -12,6 +12,7 @@
 #include "HAL/PlatformProcess.h"
 #include "HAL/Runnable.h"
 #include "HAL/RunnableThread.h"
+#include "Misc/ScopeExit.h"
 #include "Sockets.h"
 #include "Templates/UniquePtr.h"
 
@@ -142,15 +143,17 @@ bool FSecureSocketStream::LoadCertificateAuthority(FStringView InCertificatePath
 	return true;
 }
 
-bool FSecureSocketStream::SetCertificateAuthority(TArrayView<uint8> CA)
+bool FSecureSocketStream::SetCertificateAuthority(TConstArrayView<uint8> CA)
 {
 	const unsigned char* CAPointer = CA.GetData(); 
 	size_t CASize = CA.Num();
+
 	X509* DecodedCert = d2i_X509(NULL, &CAPointer, CASize);
+	ON_SCOPE_EXIT { X509_free(DecodedCert); };
 
 	if (!DecodedCert)
 	{
-		UE_LOG(LogTraceData, Error, TEXT("Failed to decode SSL certificate authority"));
+		UE_LOG(LogTraceData, Error, TEXT("Failed to decode SSL certificate authority (size: %d)"), CA.Num());
 		return false;
 	}
 
@@ -292,7 +295,7 @@ void FSecureSocketStream::SetState(EState NewState, bool bTriggerReady)
 	}
 }
 
-bool FSecureSocketStream::ConnectWithCertificate(FStringView InHost, uint16 InPort, TArrayView<uint8> InAuthToken, FStringView InCertificatePath)
+bool FSecureSocketStream::ConnectWithCertificate(FStringView InHost, uint16 InPort, TConstArrayView<uint8> InAuthToken, FStringView InCertificatePath)
 {
 	if (InCertificatePath.IsEmpty() || !LoadCertificateAuthority(InCertificatePath))
 	{
@@ -302,7 +305,7 @@ bool FSecureSocketStream::ConnectWithCertificate(FStringView InHost, uint16 InPo
 	return ConnectInternal(InHost, InPort, InAuthToken);
 }
 
-bool FSecureSocketStream::ConnectWithSelfSigned(FStringView InHost, uint16 InPort, TArrayView<uint8> InAuthToken, TArrayView<uint8> InCertificate)
+bool FSecureSocketStream::ConnectWithSelfSigned(FStringView InHost, uint16 InPort, TConstArrayView<uint8> InAuthToken, TConstArrayView<uint8> InCertificate)
 {
 	if (!SetCertificateAuthority(InCertificate))
 	{
@@ -313,7 +316,7 @@ bool FSecureSocketStream::ConnectWithSelfSigned(FStringView InHost, uint16 InPor
 	return ConnectInternal(InHost, InPort, InAuthToken);
 }
 
-bool FSecureSocketStream::ConnectWithPreSharedKey(FStringView InHost, uint16 InPort, TArrayView<uint8> InPreSharedKey, FAnsiStringView InIdentity)
+bool FSecureSocketStream::ConnectWithPreSharedKey(FStringView InHost, uint16 InPort, TConstArrayView<uint8> InPreSharedKey, FAnsiStringView InIdentity)
 {
 	if (!SetPreSharedKey())
 	{
@@ -324,7 +327,7 @@ bool FSecureSocketStream::ConnectWithPreSharedKey(FStringView InHost, uint16 InP
 	return ConnectInternal(InHost, InPort, InPreSharedKey);
 }
 
-bool FSecureSocketStream::ConnectInternal(FStringView InHost, uint16 InPort, TArrayView<uint8> InAuthToken)
+bool FSecureSocketStream::ConnectInternal(FStringView InHost, uint16 InPort, TConstArrayView<uint8> InAuthToken)
 {
 	// Create the socket object
 	AsioContext->Socket = MakeUnique<asio::ssl::stream<asio::ip::tcp::socket>>(AsioContext->Context, AsioContext->SSLContext);
@@ -409,12 +412,12 @@ void FSecureSocketStream::AsyncHandshake()
 	);
 }
 
-void FSecureSocketStream::AsyncAuthenticate()
+void FSecureSocketStream::AsyncAuthenticate(int32 Offset)
 {
 	State = EState::Authenticating;
 	AsioContext->Socket->async_write_some(
-		asio::buffer(AuthToken.GetData(), AuthToken.Num()),
-		[this] (asio::error_code Error, size_t BytesWritten)
+		asio::buffer(AuthToken.GetData() + Offset, AuthToken.Num() - Offset),
+		[this, Offset] (asio::error_code Error, size_t BytesWritten)
 		{
 			if (Error)
 			{
@@ -425,6 +428,12 @@ void FSecureSocketStream::AsyncAuthenticate()
 				SetState(EState::Failed, true);
 				return;
 			}
+			int32 CurrentOffset = Offset + BytesWritten;
+			if (CurrentOffset != AuthToken.Num())
+			{
+				AsyncAuthenticate(CurrentOffset);
+				return;
+			}
 			UE_LOG(LogTraceData, Verbose, TEXT("Authentication sent"));
 			SetState(EState::Open, true);
 			ReadyToReadEvent->Trigger();
diff --git a/Engine/Source/Developer/TraceAnalysis/Public/Trace/SecureDataStream.h b/Engine/Source/Developer/TraceAnalysis/Public/Trace/SecureDataStream.h
--- a/Engine/Source/Developer/TraceAnalysis/Public/Trace/SecureDataStream.h
+++ b/Engine/Source/Developer/TraceAnalysis/Public/Trace/SecureDataStream.h
@@ -42,7 +42,7 @@ public:
 	 * @param CertificatePath Optional path to trusted root certificate
 	 * @return True on success, false if not.
 	 */
-	UE_API bool ConnectWithCertificate(FStringView Host, uint16 Port, TArrayView<uint8> AuthToken, FStringView CertificatePath);
+	UE_API bool ConnectWithCertificate(FStringView Host, uint16 Port, TConstArrayView<uint8> AuthToken, FStringView CertificatePath);
 
 	/**
 	 * Connect to a tracing host. Specify a authorization token and an optional trusted in-memory CA certificate.
@@ -52,7 +52,7 @@ public:
 	 * @param Certificate A self-signed root certificate in DER format.
 	 * @return True on success, false if not.
 	 */
-	UE_API bool ConnectWithSelfSigned(FStringView Host, uint16 Port, TArrayView<uint8> AuthToken, TArrayView<uint8> Certificate);
+	UE_API bool ConnectWithSelfSigned(FStringView Host, uint16 Port, TConstArrayView<uint8> AuthToken, TConstArrayView<uint8> Certificate);
 
 	/**
 	 * Connect to a tracing host. Specify an authorization token used for pre-shared key verification.
@@ -62,7 +62,7 @@ public:
 	 * @param Identity Identity to use in secure handshake
 	 * @return True on success, false if not.
 	 */
-	UE_API bool ConnectWithPreSharedKey(FStringView Host, uint16 Port, TArrayView<uint8> PreSharedKey, FAnsiStringView Identity);
+	UE_API bool ConnectWithPreSharedKey(FStringView Host, uint16 Port, TConstArrayView<uint8> PreSharedKey, FAnsiStringView Identity);
 
 
 	UE_API virtual int32 Read(void* Data, uint32 Size) override;
@@ -86,7 +86,7 @@ private:
 		Failed = 0xff
 	};
 
-	bool ConnectInternal(FStringView Host, uint16 Port, TArrayView<uint8> AuthToken);
+	bool ConnectInternal(FStringView Host, uint16 Port, TConstArrayView<uint8> AuthToken);
 
 	// IInStream interface
 	UE_API virtual bool WaitUntilReady() override;
@@ -98,13 +98,13 @@ private:
 
 	// Async data functions
 	void AsyncHandshake();
-	void AsyncAuthenticate();
+	void AsyncAuthenticate(int32 Offset = 0);
 #if UE_TRACE_ANALYSIS_SSS_BUFFERED
 	void AsyncRecieveData();
 #endif
 
 	bool LoadCertificateAuthority(FStringView InCertificatePath);
-	bool SetCertificateAuthority(TArrayView<uint8> CA);
+	bool SetCertificateAuthority(TConstArrayView<uint8> CA);
 	bool SetPreSharedKey();
 	void SetState(EState NewState, bool bTriggerReady);
 
diff --git a/Engine/Source/Developer/TraceInsights/Private/Insights/InsightsManager.cpp b/Engine/Source/Developer/TraceInsights/Private/Insights/InsightsManager.cpp
--- a/Engine/Source/Developer/TraceInsights/Private/Insights/InsightsManager.cpp
+++ b/Engine/Source/Developer/TraceInsights/Private/Insights/InsightsManager.cpp
@@ -877,9 +877,11 @@ uint16 FInsightsManager::ListenForDirectTrace(const TCHAR* InStreamName, uint16
 
 ////////////////////////////////////////////////////////////////////////////////////////////////////
 
-void FInsightsManager::ConnectToSecureTrace(FStringView Host, uint16 Port, TArrayView<uint8> AuthorizationToken, FStringView CertificateFile, TArrayView<uint8> CertificateAuthority, FAnsiStringView Identity)
+void FInsightsManager::ConnectToSecureTrace(FStringView Host, uint16 Port, TConstArrayView<uint8> AuthorizationToken, FStringView CertificateFile, TConstArrayView<uint8> CertificateAuthority, FAnsiStringView Identity)
 {
 #if UE_TRACE_ANALYSIS_SSS_ENABLED
+	ResetSession();
+
 	TUniquePtr<UE::Trace::FSecureSocketStream> TraceData = MakeUnique<UE::Trace::FSecureSocketStream>();
 	bool bConnected = false;
 	if (!CertificateFile.IsEmpty())
@@ -895,27 +897,33 @@ void FInsightsManager::ConnectToSecureTrace(FStringView Host, uint16 Port, TArra
 		bConnected = TraceData->ConnectWithPreSharedKey(Host, Port, AuthorizationToken, Identity);
 	}
 
-	if (!bConnected)
+	bool bSuccess = false;
+	if (bConnected)
 	{
-		UE_LOG(LogCore, Error, TEXT("Failed to connect to secure trace session."));
-	}
 
-	constexpr uint32 TraceId = 0;
-	TStringBuilder<64> TraceName;
-	TraceName << TEXT("Secure Trace Stream from ") << Host;
+		constexpr uint32 TraceId = 0;
+		TStringBuilder<64> TraceName;
+		TraceName << TEXT("Secure Trace Stream from ") << Host;
 
-	Session = AnalysisService->StartAnalysis(TraceId, TraceName.ToString(), MoveTemp(TraceData));
+		Session = AnalysisService->StartAnalysis(TraceId, TraceName.ToString(), MoveTemp(TraceData));
 
-	if (Session)
+		if (Session)
+		{
+			CurrentTraceStreamType = ETraceStreamType::DirectTrace;
+			CurrentTraceName = TraceName;
+			CurrentTraceId = 0;
+			CurrentDirectTracePort = Port;
+			bIsSessionInfoSet = false;
+			OnSessionChanged();
+			bSuccess = true;
+		}
+	}
+	else
 	{
-		CurrentTraceStreamType = ETraceStreamType::DirectTrace;
-		CurrentTraceName = TraceName;
-		CurrentTraceId = 0;
-		CurrentDirectTracePort = Port;
-		bIsSessionInfoSet = false;
-		OnSessionChanged();
+		UE_LOG(LogCore, Error, TEXT("Failed to connect to secure trace session."));
 	}
-	else if (bAutoQuit)
+
+	if (!bSuccess && bAutoQuit)
 	{
 		RequestEngineExit(AutoQuitMsgOnFail);
 	}
diff --git a/Engine/Source/Developer/TraceInsights/Private/Insights/InsightsManager.h b/Engine/Source/Developer/TraceInsights/Private/Insights/InsightsManager.h
--- a/Engine/Source/Developer/TraceInsights/Private/Insights/InsightsManager.h
+++ b/Engine/Source/Developer/TraceInsights/Private/Insights/InsightsManager.h
@@ -264,11 +264,11 @@ public:
 	 * @param Port Port to connect to.
 	 * @param AuthorizationToken Optional authorization token to use when connecting.
 	 * @param CertificateFile Optional certificate to use when connecting. Excludes usage of CertificateAuthority.
-	 * @param CertificateAuthority Optional certificate to use as trusted. CertificateFile must be nullptr.
+	 * @param CertificateAuthority Optional certificate to use as trusted. CertificateFile should be unset to use.
 	 * @param Identity Identity string used to connect when using pre-shared keys.
 	 * @note All view-type parameters are copied and does not have to outlive this function call.
 	 */
-	void ConnectToSecureTrace(FStringView Host, uint16 Port, TArrayView<uint8> AuthorizationToken, FStringView CertificateFile = FStringView(), TArrayView<uint8> CertificateAuthority = TArrayView<uint8>(), FAnsiStringView Identity = FAnsiStringView());
+	void ConnectToSecureTrace(FStringView Host, uint16 Port, TConstArrayView<uint8> AuthorizationToken, FStringView CertificateFile = FStringView(), TConstArrayView<uint8> CertificateAuthority = TConstArrayView<uint8>(), FAnsiStringView Identity = FAnsiStringView());
 
 	/**
 	 * Creates a new analysis session instance from a provided data stream.
diff --git a/Engine/Source/Developer/TraceInsights/Private/Insights/TraceInsightsModule.cpp b/Engine/Source/Developer/TraceInsights/Private/Insights/TraceInsightsModule.cpp
--- a/Engine/Source/Developer/TraceInsights/Private/Insights/TraceInsightsModule.cpp
+++ b/Engine/Source/Developer/TraceInsights/Private/Insights/TraceInsightsModule.cpp
@@ -527,7 +527,7 @@ uint16 FTraceInsightsModule::StartAnalysisWithDirectTrace(const TCHAR* InStreamN
 }
 
 ////////////////////////////////////////////////////////////////////////////////////////////////////
-void FTraceInsightsModule::StartAnalysisWithSecureTrace(FStringView Host, uint16 Port, TArrayView<uint8> AuthorizationToken, FStringView CertificateFile, TArrayView<uint8> CertificateAuthority, FAnsiStringView Identity)
+void FTraceInsightsModule::StartAnalysisWithSecureTrace(FStringView Host, uint16 Port, TConstArrayView<uint8> AuthorizationToken, FStringView CertificateFile, TConstArrayView<uint8> CertificateAuthority, FAnsiStringView Identity)
 {
 	FInsightsManager::Get()->ConnectToSecureTrace(Host, Port, AuthorizationToken, CertificateFile, CertificateAuthority, Identity);
 	UpdateAppTitle();
diff --git a/Engine/Source/Developer/TraceInsights/Private/Insights/TraceInsightsModule.h b/Engine/Source/Developer/TraceInsights/Private/Insights/TraceInsightsModule.h
--- a/Engine/Source/Developer/TraceInsights/Private/Insights/TraceInsightsModule.h
+++ b/Engine/Source/Developer/TraceInsights/Private/Insights/TraceInsightsModule.h
@@ -65,7 +65,7 @@ public:
 	virtual void StartAnalysisForLastLiveSession(float InRetryTime = 1.0f) override;
 	virtual void StartAnalysisForTraceFile(const TCHAR* InTraceFile) override;
 	virtual uint16 StartAnalysisWithDirectTrace(const TCHAR* InStreamName = nullptr, uint16 InPort = 0) override;
-	virtual void StartAnalysisWithSecureTrace(FStringView Host, uint16 Port, TArrayView<uint8> AuthorizationToken, FStringView CertificateFile = FStringView(), TArrayView<uint8> CertificateAuthority = TArrayView<uint8>(), FAnsiStringView Identity = FAnsiStringView()) override;
+	virtual void StartAnalysisWithSecureTrace(FStringView Host, uint16 Port, TConstArrayView<uint8> AuthorizationToken, FStringView CertificateFile = FStringView(), TConstArrayView<uint8> CertificateAuthority = TConstArrayView<uint8>(), FAnsiStringView Identity = FAnsiStringView()) override;
 	virtual void StartAnalysisWithStream(TUniquePtr<UE::Trace::IInDataStream>&& InStream, const TCHAR* InStreamName = nullptr) override;
 
 	virtual void ShutdownUserInterface() override;
diff --git a/Engine/Source/Developer/TraceInsights/Public/Insights/IUnrealInsightsModule.h b/Engine/Source/Developer/TraceInsights/Public/Insights/IUnrealInsightsModule.h
--- a/Engine/Source/Developer/TraceInsights/Public/Insights/IUnrealInsightsModule.h
+++ b/Engine/Source/Developer/TraceInsights/Public/Insights/IUnrealInsightsModule.h
@@ -314,11 +314,11 @@ public:
 	 * @param Port Port to connect to.
 	 * @param AuthorizationToken Optional authorization token to use when connecting.
 	 * @param CertificateFile Optional certificate to use when connecting. Excludes usage of CertificateAuthority.
-	 * @param CertificateAuthority Optional certificate to use as trusted. CertificateFile must be nullptr.
+	 * @param CertificateAuthority Optional certificate to use as trusted. CertificateFile should be unset to use.
 	 * @param Identity Identity string used to connect when using pre-shared keys.
 	 * @note All view-type parameters are copied and does not have to outlive this function call.
 	 */
-	virtual void StartAnalysisWithSecureTrace(FStringView Host, uint16 Port, TArrayView<uint8> AuthorizationToken, FStringView CertificateFile = FStringView(), TArrayView<uint8> CertificateAuthority = TArrayView<uint8>(), FAnsiStringView Identity = FAnsiStringView()) = 0;
+	virtual void StartAnalysisWithSecureTrace(FStringView Host, uint16 Port, TConstArrayView<uint8> AuthorizationToken, FStringView CertificateFile = FStringView(), TConstArrayView<uint8> CertificateAuthority = TConstArrayView<uint8>(), FAnsiStringView Identity = FAnsiStringView()) = 0;
 
 	/**
 	 * Start analysis given a data stream.
