[Insights] TraceWriter: Added support for trace events with "DefinitionId" fields.

#Synced-CL 52292273

diff --git a/Engine/Source/Developer/TraceAnalysis/Private/TraceWriter.cpp b/Engine/Source/Developer/TraceAnalysis/Private/TraceWriter.cpp
--- a/Engine/Source/Developer/TraceAnalysis/Private/TraceWriter.cpp
+++ b/Engine/Source/Developer/TraceAnalysis/Private/TraceWriter.cpp
@@ -560,7 +560,7 @@ void FTraceWriter::Begin(ETraceWriterBeginOptions Options)
 		DeclareDefaultEvents();
 		WritePacket();
 	}
-	
+
 	if (EnumHasAnyFlags(Options, ETraceWriterBeginOptions::WriteNewTraceEvent))
 	{
 		// Write $Trace.NewTrace event.
@@ -863,7 +863,7 @@ FTraceWriterEventDeclarationBuilder& FTraceWriter::DeclareEvent(FAnsiStringView
 	FTraceWriterEventInfo* EventInfo = new FTraceWriterEventInfo(LoggerName, Name, Flags);
 	EventInfo->Id = static_cast<uint32>(EventInfos.Num());
 	EventInfos.Add(EventInfo);
-	
+
 	FTraceWriterEventDeclarationBuilder* Builder = new FTraceWriterEventDeclarationBuilder(*this, *EventInfo);
 	PendingEventDeclaration = Builder;
 	return *Builder;
@@ -966,7 +966,7 @@ void FTraceWriter::WriteEventDeclaration(const FTraceWriterEventInfo& EventInfo)
 
 	//////////////////////////////////////////////////
 	// Reserve buffer.
-	
+
 	uint8* Buffer = GetBuffer(TotalEventSize);
 	if (!Buffer)
 	{
@@ -991,7 +991,7 @@ void FTraceWriter::WriteEventDeclaration(const FTraceWriterEventInfo& EventInfo)
 	NewEvent.Flags = GetProtocolEventFlags(EventInfo.GetFlags());
 	NewEvent.LoggerNameSize = static_cast<uint8>(EventInfo.GetLoggerName().Len());
 	NewEvent.EventNameSize = static_cast<uint8>(EventInfo.GetName().Len());
-	
+
 	// Write fields.
 	uint32 Offset = 0;
 	for (uint32 FieldIndex = 0; FieldIndex < FieldCount; ++FieldIndex)
@@ -1028,7 +1028,7 @@ void FTraceWriter::WriteEventDeclaration(const FTraceWriterEventInfo& EventInfo)
 		Offset += FieldInfo.GetSize();
 	}
 
-	// Write names
+	// Write names.
 	uint8* NamesCursor = (uint8*)(NewEvent.Fields + NewEvent.FieldCount);
 	auto WriteName = [&NamesCursor](const FAnsiString& String)
 		{
@@ -1317,16 +1317,16 @@ void FTraceWriter::WriteAuxDataSegment(uint32 Index, const void* SegmentData, ui
 		return;
 	}
 
-	// Write header
+	// Write header.
 	uint32 Pack = SegmentSize << FAuxHeader::SizeShift;
 	Pack |= Index << FAuxHeader::FieldShift;
 	FMemory::Memcpy(Buffer, &Pack, sizeof(uint32)); /* FAuxHeader::Pack */
-	*Buffer = bIsImportant ?  
+	*Buffer = bIsImportant ?
 		uint8(EKnownEventUids::AuxData) : // no _UidShift !!!
 		uint8(EKnownEventUids::AuxData << EKnownEventUids::_UidShift); /* FAuxHeader::Uid */
 	Buffer += sizeof(FAuxHeader);
 
-	// Write data
+	// Write data.
 	FMemory::Memcpy(Buffer, SegmentData, SegmentSize);
 
 	AdvanceBuffer(Size);
@@ -1491,25 +1491,25 @@ uint8 FTraceWriter::GetProtocolFieldType(ETraceWriterFieldType Type)
 
 	if (EnumHasAnyFlags(Type, ETraceWriterFieldType::ReferenceFlag))
 	{
-		//switch (Type & ETraceWriterFieldType::IndexMask)
-		//{
-		//	case ETraceWriterFieldType::Reference8:  return uint8(EFieldType::Uint8);
-		//	case ETraceWriterFieldType::Reference16: return uint8(EFieldType::Uint16);
-		//	case ETraceWriterFieldType::Reference32: return uint8(EFieldType::Uint32);
-		//	case ETraceWriterFieldType::Reference64: return uint8(EFieldType::Uint64);
-		//}
+		switch (Type & (ETraceWriterFieldType::IndexMask | ETraceWriterFieldType::ReferenceFlag))
+		{
+			case ETraceWriterFieldType::Reference8:  return uint8(EFieldType::Uint8);
+			case ETraceWriterFieldType::Reference16: return uint8(EFieldType::Uint16);
+			case ETraceWriterFieldType::Reference32: return uint8(EFieldType::Uint32);
+			case ETraceWriterFieldType::Reference64: return uint8(EFieldType::Uint64);
+		}
 		return 0;
 	}
 
 	if (EnumHasAnyFlags(Type, ETraceWriterFieldType::DefinitionIdFlag))
 	{
-		//switch (Type & ETraceWriterFieldType::IndexMask)
-		//{
-		//	case ETraceWriterFieldType::DefinitionId8:  return uint8(EFieldType::Uint8);
-		//	case ETraceWriterFieldType::DefinitionId16: return uint8(EFieldType::Uint16);
-		//	case ETraceWriterFieldType::DefinitionId32: return uint8(EFieldType::Uint32);
-		//	case ETraceWriterFieldType::DefinitionId64: return uint8(EFieldType::Uint64);
-		//}
+		switch (Type & (ETraceWriterFieldType::IndexMask | ETraceWriterFieldType::DefinitionIdFlag))
+		{
+			case ETraceWriterFieldType::DefinitionId8:  return uint8(EFieldType::Uint8);
+			case ETraceWriterFieldType::DefinitionId16: return uint8(EFieldType::Uint16);
+			case ETraceWriterFieldType::DefinitionId32: return uint8(EFieldType::Uint32);
+			case ETraceWriterFieldType::DefinitionId64: return uint8(EFieldType::Uint64);
+		}
 		return 0;
 	}
 
@@ -1602,6 +1602,12 @@ ETraceWriterFieldType FTraceWriter::ConvertFieldType(const UE::Trace::IAnalyzer:
 		case EType::Reference64: WriterFieldType = ETraceWriterFieldType::ReferenceFlag | ETraceWriterFieldType::Uint64; break;
 	}
 
+	if (EnumHasAnyFlags(WriterFieldType, ETraceWriterFieldType::ReferenceFlag) &&
+		FCStringAnsi::Strcmp(FieldInfo.GetName(), "DefinitionId") == 0)
+	{
+		WriterFieldType = ETraceWriterFieldType::DefinitionIdFlag | (WriterFieldType & ~ETraceWriterFieldType::ReferenceFlag);
+	}
+
 	return WriterFieldType;
 }
 
diff --git a/Engine/Source/Developer/TraceAnalysis/Public/Trace/TraceWriter.h b/Engine/Source/Developer/TraceAnalysis/Public/Trace/TraceWriter.h
--- a/Engine/Source/Developer/TraceAnalysis/Public/Trace/TraceWriter.h
+++ b/Engine/Source/Developer/TraceAnalysis/Public/Trace/TraceWriter.h
@@ -477,6 +477,45 @@ public:
 		return Field(EventInfo.GetFieldIndex(Name), Value);
 	}
 
+	//////////////////////////////////////////////////
+	// Reference types
+
+	FTraceWriterEventBuilder& Field(uint32 Index, UE::Trace::FEventRef8 Value)
+	{
+		if (CheckValidPodField(Index, ETraceWriterFieldType::Reference8, 1))
+		{
+			Fields[Index].SetData(&Value.Id, 1);
+		}
+		return *this;
+	}
+
+	FTraceWriterEventBuilder& Field(uint32 Index, UE::Trace::FEventRef16 Value)
+	{
+		if (CheckValidPodField(Index, ETraceWriterFieldType::Reference16, 2))
+		{
+			Fields[Index].SetData(&Value.Id, 2);
+		}
+		return *this;
+	}
+
+	FTraceWriterEventBuilder& Field(uint32 Index, UE::Trace::FEventRef32 Value)
+	{
+		if (CheckValidPodField(Index, ETraceWriterFieldType::Reference32, 4))
+		{
+			Fields[Index].SetData(&Value.Id, 4);
+		}
+		return *this;
+	}
+
+	FTraceWriterEventBuilder& Field(uint32 Index, UE::Trace::FEventRef64 Value)
+	{
+		if (CheckValidPodField(Index, ETraceWriterFieldType::Reference64, 8))
+		{
+			Fields[Index].SetData(&Value.Id, 8);
+		}
+		return *this;
+	}
+
 	//////////////////////////////////////////////////
 	// AnsiString
 
