[Insights]
* TraceAnalyzer: Adedd -analysis_log_level=<level> command line parameter to control the debug logs from the trace analysis engine; where <level> can be 1=minimum, 2=normal (default), 3=verbose or 4=verbose+.
* TraceAnalysis: Added UE::Trace::Analysis::GDebugLevel to control the debug log level used when UE_TRACE_ANALYSIS_DEBUG is enabled. Also added UE_TRACE_ANALYSIS_DEBUG_LOG_EX(level, formt, ...). The macro UE_TRACE_ANALYSIS_DEBUG_LEVEL is now removed.

#Synced-CL 52381187

diff --git a/Engine/Source/Developer/TraceAnalysis/Private/Analysis/Engine.cpp b/Engine/Source/Developer/TraceAnalysis/Private/Analysis/Engine.cpp
--- a/Engine/Source/Developer/TraceAnalysis/Private/Analysis/Engine.cpp
+++ b/Engine/Source/Developer/TraceAnalysis/Private/Analysis/Engine.cpp
@@ -1578,17 +1578,12 @@ void FAnalyzerHub::OnNewType(const FTypeRegistry::FTypeInfo* TypeInfo)
 
 #if UE_TRACE_ANALYSIS_DEBUG
 	const IAnalyzer::FEventTypeInfo& EventTypeInfo = *(IAnalyzer::FEventTypeInfo*)TypeInfo;
-#if UE_TRACE_ANALYSIS_DEBUG_API
 	UE_TRACE_ANALYSIS_DEBUG_LOG("NewEvent: Uid=%u (%s.%s) Flags=%s%s%s%s",
 		EventTypeInfo.GetId(), EventTypeInfo.GetLoggerName(), EventTypeInfo.GetName(),
 		EventTypeInfo.IsDefinition() ? "Definition|" : "",
 		EventTypeInfo.IsImportant() ? "Important|" : "",
 		EventTypeInfo.MaybeHasAux() ? "MaybeHasAux|" : "",
 		EventTypeInfo.IsNoSync() ? "NoSync" : "Sync");
-#else // UE_TRACE_ANALYSIS_DEBUG_API
-	UE_TRACE_ANALYSIS_DEBUG_LOG("NewEvent: Uid=%u (%s.%s)",
-		EventTypeInfo.GetId(), EventTypeInfo.GetLoggerName(), EventTypeInfo.GetName());
-#endif // UE_TRACE_ANALYSIS_DEBUG_API
 #endif // UE_TRACE_ANALYSIS_DEBUG
 
 	// Inform routes that a new event has been declared.
@@ -2081,7 +2076,7 @@ public:
 	void				DebugLogLeaveScopeEvent(uint32 Uid, uint64 RelativeTimestamp, uint32 EventSize);
 	void				DebugLogLeaveScopeAEvent(uint32 Uid, uint64 AbsoluteTimestamp, uint32 EventSize);
 	void				DebugLogLeaveScopeBEvent(uint32 Uid, uint64 BaseRelativeTimestamp, uint32 EventSize);
-#endif // UE_TRACE_ANALYSIS_DEBUG_API
+#endif // UE_TRACE_ANALYSIS_DEBUG
 
 private:
 	void				DispatchLeaveScope();
@@ -2284,7 +2279,7 @@ bool FAnalysisBridge::ShouldWriteDebugLog()
 {
 	// Example for limiting debug logs to first 100K events...
 	//return (State.TotalEventCount < 100000);
-	return true;
+	return UE::Trace::Analysis::GDebugLevel >= 2;
 }
 
 ////////////////////////////////////////////////////////////////////////////////
@@ -2294,7 +2289,6 @@ void FAnalysisBridge::DebugLogEvent(const FTypeRegistry::FTypeInfo* TypeInfo, ui
 	uint32 EventSize = FixedSize + AuxSize;
 	State.TotalEventSize += EventSize;
 
-#if UE_TRACE_ANALYSIS_DEBUG_LEVEL >= 2
 	const bool bDebugLog = ShouldWriteDebugLog();
 	if (bDebugLog)
 	{
@@ -2317,29 +2311,26 @@ void FAnalysisBridge::DebugLogEvent(const FTypeRegistry::FTypeInfo* TypeInfo, ui
 			UE_TRACE_ANALYSIS_DEBUG_Appendf(" (%u + %u)", FixedSize, AuxSize);
 		}
 	}
-#endif // UE_TRACE_ANALYSIS_DEBUG_LEVEL
 
 	if (TypeInfo->IsSync())
 	{
 		++State.SyncEventCount;
 		State.SyncEventSize += EventSize;
-#if UE_TRACE_ANALYSIS_DEBUG_LEVEL >= 2
+
 		if (bDebugLog)
 		{
 			UE_TRACE_ANALYSIS_DEBUG_Appendf(" Sync=%u", Serial);
 		}
-#endif // UE_TRACE_ANALYSIS_DEBUG_LEVEL
 	}
 	else if (TypeInfo->IsImportant())
 	{
 		++State.ImportantNoSyncEventCount;
 		State.ImportantNoSyncEventSize += EventSize;
-#if UE_TRACE_ANALYSIS_DEBUG_LEVEL >= 2
+
 		if (bDebugLog)
 		{
 			UE_TRACE_ANALYSIS_DEBUG_Append(" Important");
 		}
-#endif // UE_TRACE_ANALYSIS_DEBUG_LEVEL
 	}
 	else
 	{
@@ -2347,12 +2338,10 @@ void FAnalysisBridge::DebugLogEvent(const FTypeRegistry::FTypeInfo* TypeInfo, ui
 		State.OtherNoSyncEventSize += EventSize;
 	}
 
-#if UE_TRACE_ANALYSIS_DEBUG_LEVEL >= 2
 	if (bDebugLog)
 	{
 		UE_TRACE_ANALYSIS_DEBUG_EndStringBuilder();
 	}
-#endif // UE_TRACE_ANALYSIS_DEBUG_LEVEL
 }
 
 ////////////////////////////////////////////////////////////////////////////////
@@ -2362,14 +2351,13 @@ void FAnalysisBridge::DebugLogNewEvent(uint32 Uid, const FTypeRegistry::FTypeInf
 	++State.NewEventCount;
 	State.TotalEventSize += EventSize;
 	State.NewEventSize += EventSize;
-#if UE_TRACE_ANALYSIS_DEBUG_LEVEL >= 2
+
 	const bool bDebugLog = ShouldWriteDebugLog();
 	if (bDebugLog)
 	{
 		const uint32 Tid = ETransportTid::Importants; // GetActiveThreadId()
 		UE_TRACE_ANALYSIS_DEBUG_LOG("[EVENT %llu] Tid=%u Uid=%u (NewEvent) Size=%u", State.TotalEventCount, Tid, Uid, EventSize);
 	}
-#endif // UE_TRACE_ANALYSIS_DEBUG_LEVEL
 }
 
 ////////////////////////////////////////////////////////////////////////////////
@@ -2379,14 +2367,13 @@ void FAnalysisBridge::DebugLogEnterScopeEvent(uint32 Uid, uint32 EventSize)
 	++State.EnterScopeEventCount;
 	State.TotalEventSize += EventSize;
 	State.EnterScopeEventSize += EventSize;
-#if UE_TRACE_ANALYSIS_DEBUG_LEVEL >= 2
+
 	const bool bDebugLog = ShouldWriteDebugLog();
 	if (bDebugLog)
 	{
 		UE_TRACE_ANALYSIS_DEBUG_LOG("[EVENT %llu] Tid=%u Uid=%u (EnterScope) Size=%u",
 			State.TotalEventCount, GetActiveThreadId(), Uid, EventSize);
 	}
-#endif // UE_TRACE_ANALYSIS_DEBUG_LEVEL
 }
 
 ////////////////////////////////////////////////////////////////////////////////
@@ -2396,7 +2383,7 @@ void FAnalysisBridge::DebugLogEnterScopeEvent(uint32 Uid, uint64 RelativeTimesta
 	++State.EnterScopeTEventCount;
 	State.TotalEventSize += EventSize;
 	State.EnterScopeTEventSize += EventSize;
-#if UE_TRACE_ANALYSIS_DEBUG_LEVEL >= 2
+
 	const bool bDebugLog = ShouldWriteDebugLog();
 	if (bDebugLog)
 	{
@@ -2406,7 +2393,6 @@ void FAnalysisBridge::DebugLogEnterScopeEvent(uint32 Uid, uint64 RelativeTimesta
 			RelativeTimestamp, State.Timing.BaseTimestamp + ThreadInfo->PrevTimestamp + RelativeTimestamp,
 			EventSize);
 	}
-#endif // UE_TRACE_ANALYSIS_DEBUG_LEVEL
 }
 
 ////////////////////////////////////////////////////////////////////////////////
@@ -2416,7 +2402,7 @@ void FAnalysisBridge::DebugLogEnterScopeAEvent(uint32 Uid, uint64 AbsoluteTimest
 	++State.EnterScopeTEventCount;
 	State.TotalEventSize += EventSize;
 	State.EnterScopeTEventSize += EventSize;
-#if UE_TRACE_ANALYSIS_DEBUG_LEVEL >= 2
+
 	const bool bDebugLog = ShouldWriteDebugLog();
 	if (bDebugLog)
 	{
@@ -2425,7 +2411,6 @@ void FAnalysisBridge::DebugLogEnterScopeAEvent(uint32 Uid, uint64 AbsoluteTimest
 			AbsoluteTimestamp,
 			EventSize);
 	}
-#endif // UE_TRACE_ANALYSIS_DEBUG_LEVEL
 }
 
 ////////////////////////////////////////////////////////////////////////////////
@@ -2435,7 +2420,7 @@ void FAnalysisBridge::DebugLogEnterScopeBEvent(uint32 Uid, uint64 BaseAbsoluteTi
 	++State.EnterScopeTEventCount;
 	State.TotalEventSize += EventSize;
 	State.EnterScopeTEventSize += EventSize;
-#if UE_TRACE_ANALYSIS_DEBUG_LEVEL >= 2
+
 	const bool bDebugLog = ShouldWriteDebugLog();
 	if (bDebugLog)
 	{
@@ -2444,7 +2429,6 @@ void FAnalysisBridge::DebugLogEnterScopeBEvent(uint32 Uid, uint64 BaseAbsoluteTi
 			BaseAbsoluteTimestamp, State.Timing.BaseTimestamp + BaseAbsoluteTimestamp,
 			EventSize);
 	}
-#endif // UE_TRACE_ANALYSIS_DEBUG_LEVEL
 }
 
 ////////////////////////////////////////////////////////////////////////////////
@@ -2454,14 +2438,13 @@ void FAnalysisBridge::DebugLogLeaveScopeEvent(uint32 Uid, uint32 EventSize)
 	++State.LeaveScopeEventCount;
 	State.TotalEventSize += EventSize;
 	State.LeaveScopeEventSize += EventSize;
-#if UE_TRACE_ANALYSIS_DEBUG_LEVEL >= 2
+
 	const bool bDebugLog = ShouldWriteDebugLog();
 	if (bDebugLog)
 	{
 		UE_TRACE_ANALYSIS_DEBUG_LOG("[EVENT %llu] Tid=%u Uid=%u (LeaveScope) Size=%u",
 			State.TotalEventCount, GetActiveThreadId(), Uid, EventSize);
 	}
-#endif // UE_TRACE_ANALYSIS_DEBUG_LEVEL
 }
 
 ////////////////////////////////////////////////////////////////////////////////
@@ -2471,7 +2454,7 @@ void FAnalysisBridge::DebugLogLeaveScopeEvent(uint32 Uid, uint64 RelativeTimesta
 	++State.LeaveScopeTEventCount;
 	State.TotalEventSize += EventSize;
 	State.LeaveScopeTEventSize += EventSize;
-#if UE_TRACE_ANALYSIS_DEBUG_LEVEL >= 2
+
 	const bool bDebugLog = ShouldWriteDebugLog();
 	if (bDebugLog)
 	{
@@ -2481,7 +2464,6 @@ void FAnalysisBridge::DebugLogLeaveScopeEvent(uint32 Uid, uint64 RelativeTimesta
 			RelativeTimestamp, State.Timing.BaseTimestamp + ThreadInfo->PrevTimestamp + RelativeTimestamp,
 			EventSize);
 	}
-#endif // UE_TRACE_ANALYSIS_DEBUG_LEVEL
 }
 
 ////////////////////////////////////////////////////////////////////////////////
@@ -2491,7 +2473,7 @@ void FAnalysisBridge::DebugLogLeaveScopeAEvent(uint32 Uid, uint64 AbsoluteTimest
 	++State.LeaveScopeTEventCount;
 	State.TotalEventSize += EventSize;
 	State.LeaveScopeTEventSize += EventSize;
-#if UE_TRACE_ANALYSIS_DEBUG_LEVEL >= 2
+
 	const bool bDebugLog = ShouldWriteDebugLog();
 	if (bDebugLog)
 	{
@@ -2500,7 +2482,6 @@ void FAnalysisBridge::DebugLogLeaveScopeAEvent(uint32 Uid, uint64 AbsoluteTimest
 			AbsoluteTimestamp,
 			EventSize);
 	}
-#endif // UE_TRACE_ANALYSIS_DEBUG_LEVEL
 }
 
 ////////////////////////////////////////////////////////////////////////////////
@@ -2510,7 +2491,7 @@ void FAnalysisBridge::DebugLogLeaveScopeBEvent(uint32 Uid, uint64 BaseAbsoluteTi
 	++State.LeaveScopeTEventCount;
 	State.TotalEventSize += EventSize;
 	State.LeaveScopeTEventSize += EventSize;
-#if UE_TRACE_ANALYSIS_DEBUG_LEVEL >= 2
+
 	const bool bDebugLog = ShouldWriteDebugLog();
 	if (bDebugLog)
 	{
@@ -2519,7 +2500,6 @@ void FAnalysisBridge::DebugLogLeaveScopeBEvent(uint32 Uid, uint64 BaseAbsoluteTi
 			BaseAbsoluteTimestamp, State.Timing.BaseTimestamp + BaseAbsoluteTimestamp,
 			EventSize);
 	}
-#endif // UE_TRACE_ANALYSIS_DEBUG_LEVEL
 }
 #endif // UE_TRACE_ANALYSIS_DEBUG
 
@@ -3087,9 +3067,7 @@ int32 FProtocol2Stage::OnData(FStreamReader& Reader, FAnalysisBridge& Bridge)
 					const auto* HeaderV1 = (Protocol1::FEventHeader*)Header;
 					if (HeaderV1->Serial != (Serial.Value & Serial.Mask))
 					{
-#if UE_TRACE_ANALYSIS_DEBUG && UE_TRACE_ANALYSIS_DEBUG_LEVEL >= 2
-						UE_TRACE_ANALYSIS_DEBUG_LOG("Tid=%u --> EventSerial=%u", Bridge.GetActiveThreadId(), HeaderV1->Serial);
-#endif
+						UE_TRACE_ANALYSIS_DEBUG_LOG_EX(2, "Tid=%u --> EventSerial=%u", Bridge.GetActiveThreadId(), HeaderV1->Serial);
 						return HeaderV1->Serial;
 					}
 					BlockSize += sizeof(*HeaderV1);
@@ -3103,9 +3081,7 @@ int32 FProtocol2Stage::OnData(FStreamReader& Reader, FAnalysisBridge& Bridge)
 					uint32 EventSerial = HeaderSync->SerialLow|(uint32(HeaderSync->SerialHigh) << 16);
 					if (EventSerial != (Serial.Value & Serial.Mask))
 					{
-#if UE_TRACE_ANALYSIS_DEBUG && UE_TRACE_ANALYSIS_DEBUG_LEVEL >= 2
-						UE_TRACE_ANALYSIS_DEBUG_LOG("Tid=%u --> EventSerial=%u", Bridge.GetActiveThreadId(), EventSerial);
-#endif
+						UE_TRACE_ANALYSIS_DEBUG_LOG_EX(2, "Tid=%u --> EventSerial=%u", Bridge.GetActiveThreadId(), EventSerial);
 						return EventSerial;
 					}
 					BlockSize += sizeof(*HeaderSync);
@@ -3343,9 +3319,7 @@ int32 FProtocol4Stage::OnDataImpl(FStreamReader& Reader, FAnalysisBridge& Bridge
 		uint32 EventSerial = HeaderSync->SerialLow|(uint32(HeaderSync->SerialHigh) << 16);
 		if (EventSerial != (Serial.Value & Serial.Mask))
 		{
-#if UE_TRACE_ANALYSIS_DEBUG && UE_TRACE_ANALYSIS_DEBUG_LEVEL >= 2
-			UE_TRACE_ANALYSIS_DEBUG_LOG("Tid=%u --> EventSerial=%u", Bridge.GetActiveThreadId(), EventSerial);
-#endif
+			UE_TRACE_ANALYSIS_DEBUG_LOG_EX(2, "Tid=%u --> EventSerial=%u", Bridge.GetActiveThreadId(), EventSerial);
 			return ~EventSerial;
 		}
 		BlockSize += sizeof(*HeaderSync);
@@ -3907,9 +3881,7 @@ FProtocol5Stage::EStatus FProtocol5Stage::OnDataNormal(const FMachineContext& Co
 	{
 		uint32 NumEventDescs = EventDescs.Num();
 
-#if UE_TRACE_ANALYSIS_DEBUG && UE_TRACE_ANALYSIS_DEBUG_LEVEL >= 3
-		UE_TRACE_ANALYSIS_DEBUG_LOG("Thread:%u Id:%u", i, Transport.GetThreadId(i));
-#endif
+		UE_TRACE_ANALYSIS_DEBUG_LOG_EX(3, "Thread:%u Id:%u", i, Transport.GetThreadId(i));
 
 		// Extract all the events in the stream for this thread
 		FStreamReader* ThreadReader = Transport.GetThreadStream(i);
@@ -3956,9 +3928,7 @@ FProtocol5Stage::EStatus FProtocol5Stage::OnDataNormal(const FMachineContext& Co
 			EventDescHeap.Add(Out);
 		}
 
-#if UE_TRACE_ANALYSIS_DEBUG && UE_TRACE_ANALYSIS_DEBUG_LEVEL >= 3
-		UE_TRACE_ANALYSIS_DEBUG_LOG("Thread:%u bNotEnoughData:%d", i, bNotEnoughData);
-#endif
+		UE_TRACE_ANALYSIS_DEBUG_LOG_EX(3, "Thread:%u bNotEnoughData:%d", i, bNotEnoughData);
 	}
 
 	// Now EventDescs is stable we can convert the indices into pointers
@@ -3987,9 +3957,7 @@ FProtocol5Stage::EStatus FProtocol5Stage::OnDataNormal(const FMachineContext& Co
 			return EStatus::Error;
 		}
 		NumAvailableEvents += SerialGaps.Num(); // serial gaps are detected during DispatchNormalEvents
-#if UE_TRACE_ANALYSIS_DEBUG && UE_TRACE_ANALYSIS_DEBUG_LEVEL >= 2
-		UE_TRACE_ANALYSIS_DEBUG_LOG("Dispatched %d normal events (%d --> %d)", NumDispatchedEvents, NumAvailableEvents, NumAvailableEvents - NumDispatchedEvents);
-#endif
+		UE_TRACE_ANALYSIS_DEBUG_LOG_EX(2, "Dispatched %d normal events (%d --> %d)", NumDispatchedEvents, NumAvailableEvents, NumAvailableEvents - NumDispatchedEvents);
 		NumAvailableEvents -= NumDispatchedEvents;
 		check(NumAvailableEvents >= 0);
 
@@ -4023,20 +3991,23 @@ FProtocol5Stage::EStatus FProtocol5Stage::OnDataNormal(const FMachineContext& Co
 		if (!bSkipSerialError)
 		{
 			bSkipSerialError = true;
-#if UE_TRACE_ANALYSIS_DEBUG && UE_TRACE_ANALYSIS_DEBUG_LEVEL >= 2
-			for (const auto& Stream : EventDescHeap)
+#if UE_TRACE_ANALYSIS_DEBUG
+			if (UE::Trace::Analysis::GDebugLevel >= 2)
 			{
-				const FEventDesc* EventDesc = Stream.EventDescs;
-				while (EventDesc->Serial != ESerial::Terminal)
+				for (const auto& Stream : EventDescHeap)
 				{
-					if (EventDesc->Serial == NextSerial)
+					const FEventDesc* EventDesc = Stream.EventDescs;
+					while (EventDesc->Serial != ESerial::Terminal)
 					{
-						UE_TRACE_ANALYSIS_DEBUG_LOG("Found next serial on thread %u (event %d)", Stream.ThreadId, (int32)(EventDesc - Stream.EventDescs));
+						if (EventDesc->Serial == NextSerial)
+						{
+							UE_TRACE_ANALYSIS_DEBUG_LOG("Found next serial on thread %u (event %d)", Stream.ThreadId, (int32)(EventDesc - Stream.EventDescs));
+						}
+						++EventDesc;
 					}
-					++EventDesc;
 				}
 			}
-#endif
+#endif // UE_TRACE_ANALYSIS_DEBUG
 			Context.EmitMessagef(
 				EAnalysisMessageSeverity::Error,
 				TEXT("Trace analysis accumulated too much data (%d parsed events) and will start to skip the missing serial sync events!"), NumAvailableEvents);
@@ -4056,9 +4027,7 @@ FProtocol5Stage::EStatus FProtocol5Stage::OnDataNormal(const FMachineContext& Co
 			{
 				return EStatus::Error;
 			}
-#if UE_TRACE_ANALYSIS_DEBUG && UE_TRACE_ANALYSIS_DEBUG_LEVEL >= 2
-			UE_TRACE_ANALYSIS_DEBUG_LOG("Skipped serials and dispatched %d normal events (%d --> %d)", NumDispatchedEvents, NumAvailableEvents, NumAvailableEvents - NumDispatchedEvents);
-#endif
+			UE_TRACE_ANALYSIS_DEBUG_LOG_EX(2, "Skipped serials and dispatched %d normal events (%d --> %d)", NumDispatchedEvents, NumAvailableEvents, NumAvailableEvents - NumDispatchedEvents);
 			NumAvailableEvents -= NumDispatchedEvents;
 			check(NumAvailableEvents >= 0);
 		}
@@ -4084,9 +4053,7 @@ FProtocol5Stage::EStatus FProtocol5Stage::OnDataNormal(const FMachineContext& Co
 			{
 				return EStatus::Error;
 			}
-#if UE_TRACE_ANALYSIS_DEBUG && UE_TRACE_ANALYSIS_DEBUG_LEVEL >= 2
-			UE_TRACE_ANALYSIS_DEBUG_LOG("Skipped serials and dispatched %d normal events (%d --> %d)", NumDispatchedEvents, NumAvailableEvents, NumAvailableEvents - NumDispatchedEvents);
-#endif
+			UE_TRACE_ANALYSIS_DEBUG_LOG_EX(2, "Skipped serials and dispatched %d normal events (%d --> %d)", NumDispatchedEvents, NumAvailableEvents, NumAvailableEvents - NumDispatchedEvents);
 			NumAvailableEvents -= NumDispatchedEvents;
 			check(NumAvailableEvents >= 0);
 
@@ -4131,9 +4098,7 @@ FProtocol5Stage::EStatus FProtocol5Stage::OnDataNormal(const FMachineContext& Co
 ////////////////////////////////////////////////////////////////////////////////
 int32 FProtocol5Stage::DispatchNormalEvents(const FMachineContext& Context, TArray<FEventDescStream>& EventDescHeap)
 {
-#if UE_TRACE_ANALYSIS_DEBUG && UE_TRACE_ANALYSIS_DEBUG_LEVEL >= 2
-	UE_TRACE_ANALYSIS_DEBUG_LOG("Queued event descs: %d", EventDescs.Num());
-#endif
+	UE_TRACE_ANALYSIS_DEBUG_LOG_EX(2, "Queued event descs: %d", EventDescs.Num());
 
 	// Process leading unsynchronized events so that each stream starts with a synchronized event.
 	int32 NumDispatchedUnsyncEvents = 0;
@@ -4179,7 +4144,8 @@ int32 FProtocol5Stage::DispatchNormalEvents(const FMachineContext& Context, TArr
 	// A min-heap is used to peel off groups of events by lowest serial
 	EventDescHeap.Heapify(FSerialDistancePredicate{NextSerial});
 
-#if UE_TRACE_ANALYSIS_DEBUG && UE_TRACE_ANALYSIS_DEBUG_LEVEL >= 2
+#if UE_TRACE_ANALYSIS_DEBUG
+	if (UE::Trace::Analysis::GDebugLevel >= 2)
 	{
 		const FEventDescStream& TopStream = EventDescHeap.HeapTop();
 		UE_TRACE_ANALYSIS_DEBUG_LOG("NextSerial=%u LowestSerial=%d (Tid=%u)", NextSerial, TopStream.EventDescs[0].Serial, TopStream.ThreadId);
@@ -4312,44 +4278,44 @@ int32 FProtocol5Stage::DispatchEvents(const FMachineContext& Context, TArray<FEv
 			}
 			else
 			{
-#if UE_TRACE_ANALYSIS_DEBUG_LEVEL >= 2
-				UE_TRACE_ANALYSIS_DEBUG_LOG("Lowest serial %d (Tid=%u Uid=%u) is not low enough (NextSerial=%u; %d event descs)", EndDesc->Serial, Stream.ThreadId, uint32(EndDesc->Uid), NextSerial, EventDescs.Num());
-#endif
+				UE_TRACE_ANALYSIS_DEBUG_LOG_EX(2, "Lowest serial %d (Tid=%u Uid=%u) is not low enough (NextSerial=%u; %d event descs)", EndDesc->Serial, Stream.ThreadId, uint32(EndDesc->Uid), NextSerial, EventDescs.Num());
 			}
-#if UE_TRACE_ANALYSIS_DEBUG_LEVEL >= 2
-			UE_TRACE_ANALYSIS_DEBUG_LOG("Available streams (SerialWrappedCount=%u):", State.SerialWrappedCount);
-			for (const FEventDescStream& EventDescStream : EventDescHeap)
+
+			if (UE::Trace::Analysis::GDebugLevel >= 2)
 			{
-				uint32 BufferSize = 0;
-				uint32 DataSize = 0; // parsed data size + remaining data size
-				uint32 RemainingSize = 0;
-				for (uint32 i = 0, n = Transport.GetThreadCount(); i < n; ++i)
+				UE_TRACE_ANALYSIS_DEBUG_LOG("Available streams (SerialWrappedCount=%u):", State.SerialWrappedCount);
+				for (const FEventDescStream& EventDescStream : EventDescHeap)
 				{
-					FStreamBuffer* ThreadReader = (FStreamBuffer*)Transport.GetThreadStream(i);
-					uint32 ThreadId = Transport.GetThreadId(i);
-					if (ThreadId == EventDescStream.ThreadId)
+					uint32 BufferSize = 0;
+					uint32 DataSize = 0; // parsed data size + remaining data size
+					uint32 RemainingSize = 0;
+					for (uint32 i = 0, n = Transport.GetThreadCount(); i < n; ++i)
 					{
-						BufferSize = ThreadReader->GetBufferSize();
-						const FEventDesc& EventDesc = EventDescStream.EventDescs[0];
-						uint32 HeaderSize = 1 + EventDesc.bTwoByteUid + (ESerial::Bits / 8);
-						DataSize = ThreadReader->GetBacktrackSize(EventDesc.Data - HeaderSize);
-						RemainingSize = ThreadReader->GetRemaining();
-						break;
+						FStreamBuffer* ThreadReader = (FStreamBuffer*)Transport.GetThreadStream(i);
+						uint32 ThreadId = Transport.GetThreadId(i);
+						if (ThreadId == EventDescStream.ThreadId)
+						{
+							BufferSize = ThreadReader->GetBufferSize();
+							const FEventDesc& EventDesc = EventDescStream.EventDescs[0];
+							uint32 HeaderSize = 1 + EventDesc.bTwoByteUid + (ESerial::Bits / 8);
+							DataSize = ThreadReader->GetBacktrackSize(EventDesc.Data - HeaderSize);
+							RemainingSize = ThreadReader->GetRemaining();
+							break;
+						}
 					}
+					UE_TRACE_ANALYSIS_DEBUG_BeginStringBuilder();
+					UE_TRACE_ANALYSIS_DEBUG_Appendf("  Tid=%u : Serial=%u BufferSize=%u DataSize=%u", EventDescStream.ThreadId, EventDescStream.EventDescs->Serial, BufferSize, DataSize);
+					if (RemainingSize != 0)
+					{
+						UE_TRACE_ANALYSIS_DEBUG_Appendf(" (%u + %u)", DataSize - RemainingSize, RemainingSize);
+					}
+					if (EventDescStream.EventDescs->Serial == NextSerial)
+					{
+						UE_TRACE_ANALYSIS_DEBUG_Append(" (next)");
+					}
+					UE_TRACE_ANALYSIS_DEBUG_EndStringBuilder();
 				}
-				UE_TRACE_ANALYSIS_DEBUG_BeginStringBuilder();
-				UE_TRACE_ANALYSIS_DEBUG_Appendf("  Tid=%u : Serial=%u BufferSize=%u DataSize=%u", EventDescStream.ThreadId, EventDescStream.EventDescs->Serial, BufferSize, DataSize);
-				if (RemainingSize != 0)
-				{
-					UE_TRACE_ANALYSIS_DEBUG_Appendf(" (%u + %u)", DataSize - RemainingSize, RemainingSize);
-				}
-				if (EventDescStream.EventDescs->Serial == NextSerial)
-				{
-					UE_TRACE_ANALYSIS_DEBUG_Append(" (next)");
-				}
-				UE_TRACE_ANALYSIS_DEBUG_EndStringBuilder();
 			}
-#endif // UE_TRACE_ANALYSIS_DEBUG_LEVEL
 #endif // UE_TRACE_ANALYSIS_DEBUG
 
 #if 0
@@ -4431,8 +4397,11 @@ int32 FProtocol5Stage::ParseEvents(FStreamReader& Reader, EventDescArray& OutEve
 			return Size;
 		}
 
-#if UE_TRACE_ANALYSIS_DEBUG && UE_TRACE_ANALYSIS_DEBUG_LEVEL >= 3
-		PrintParsedEvent(OutEventDescs.Num(), EventDesc, Size);
+#if UE_TRACE_ANALYSIS_DEBUG
+		if (UE::Trace::Analysis::GDebugLevel >= 3)
+		{
+			PrintParsedEvent(OutEventDescs.Num(), EventDesc, Size);
+		}
 #endif // UE_TRACE_ANALYSIS_DEBUG
 
 		OutEventDescs.Add(EventDesc);
@@ -4452,9 +4421,7 @@ int32 FProtocol5Stage::ParseEvents(FStreamReader& Reader, EventDescArray& OutEve
 
 			if (Ok == 0)
 			{
-#if UE_TRACE_ANALYSIS_DEBUG && UE_TRACE_ANALYSIS_DEBUG_LEVEL >= 2
-				UE_TRACE_ANALYSIS_DEBUG_LOG("Warning: Incomplete aux stack (Uid=%u)! Rewind %d parsed events.", EventDesc.Uid, OutEventDescs.Num() - RewindDescsNum);
-#endif
+				UE_TRACE_ANALYSIS_DEBUG_LOG_EX(2, "Warning: Incomplete aux stack (Uid=%u)! Rewind %d parsed events.", EventDesc.Uid, OutEventDescs.Num() - RewindDescsNum);
 				OutEventDescs.SetNum(RewindDescsNum, EAllowShrinking::No);
 				Reader.RestoreMark(RewindMark);
 				break;
@@ -4496,8 +4463,11 @@ int32 FProtocol5Stage::ParseEventsWithAux(FStreamReader& Reader, EventDescArray&
 			return Size;
 		}
 
-#if UE_TRACE_ANALYSIS_DEBUG && UE_TRACE_ANALYSIS_DEBUG_LEVEL >= 3
-		PrintParsedEvent(OutEventDescs.Num(), EventDesc, Size);
+#if UE_TRACE_ANALYSIS_DEBUG
+		if (UE::Trace::Analysis::GDebugLevel >= 3)
+		{
+			PrintParsedEvent(OutEventDescs.Num(), EventDesc, Size);
+		}
 #endif // UE_TRACE_ANALYSIS_DEBUG
 
 		Reader.Advance(Size);
@@ -4567,9 +4537,7 @@ int32 FProtocol5Stage::ParseEventsWithAux(FStreamReader& Reader, EventDescArray&
 	if (bUnsorted)
 	{
 		uint32 NumDescs = OutEventDescs.Num() - FirstDescIndex;
-#if UE_TRACE_ANALYSIS_DEBUG && UE_TRACE_ANALYSIS_DEBUG_LEVEL >= 3
-		UE_TRACE_ANALYSIS_DEBUG_LOG("Sorting %u event descs", NumDescs);
-#endif
+		UE_TRACE_ANALYSIS_DEBUG_LOG_EX(3, "Sorting %u event descs", NumDescs);
 		TArrayView<FEventDesc> DescsView(OutEventDescs.GetData() + FirstDescIndex, NumDescs);
 		Algo::StableSort(
 			DescsView,
@@ -4808,11 +4776,14 @@ void FProtocol5Stage::DetectSerialGaps(TArray<FEventDescStream>& EventDescHeap)
 			return;
 		}
 
-#if UE_TRACE_ANALYSIS_DEBUG && UE_TRACE_ANALYSIS_DEBUG_LEVEL >= 2
-		for (uint32 GapIndex = 0; GapIndex < GapCount; ++GapIndex)
+#if UE_TRACE_ANALYSIS_DEBUG
+		if (UE::Trace::Analysis::GDebugLevel >= 2)
 		{
-			const FEventDesc& SerialGap = SerialGaps[GapIndex];
-			UE_TRACE_ANALYSIS_DEBUG_LOG("  gap %u +%u", SerialGap.Serial, SerialGap.GapLength);
+			for (uint32 GapIndex = 0; GapIndex < GapCount; ++GapIndex)
+			{
+				const FEventDesc& SerialGap = SerialGaps[GapIndex];
+				UE_TRACE_ANALYSIS_DEBUG_LOG("  gap %u +%u", SerialGap.Serial, SerialGap.GapLength);
+			}
 		}
 #endif
 
@@ -4901,17 +4872,13 @@ int32 FProtocol5Stage::DispatchEvents(const FMachineContext& Context, const FEve
 
 	FAuxDataCollector AuxCollector;
 
-#if UE_TRACE_ANALYSIS_DEBUG && UE_TRACE_ANALYSIS_DEBUG_LEVEL >= 3
-	UE_TRACE_ANALYSIS_DEBUG_LOG("Dispatch run of %u consecutive events (Tid=%u)", Count, Context.Bridge.GetActiveThreadId());
-#endif
+	UE_TRACE_ANALYSIS_DEBUG_LOG_EX(3, "Dispatch run of %u consecutive events (Tid=%u)", Count, Context.Bridge.GetActiveThreadId());
 
 	for (const FEventDesc* Cursor = EventDesc, *End = EventDesc + Count; Cursor < End; ++Cursor)
 	{
 		uint32 Uid = uint32(Cursor->Uid);
 
-#if UE_TRACE_ANALYSIS_DEBUG && UE_TRACE_ANALYSIS_DEBUG_LEVEL >= 4
-		UE_TRACE_ANALYSIS_DEBUG_LOG("Dispatch event (Uid=%u)...", Uid);
-#endif
+		UE_TRACE_ANALYSIS_DEBUG_LOG_EX(4, "Dispatch event (Uid=%u)...", Uid);
 
 		if (DispatchKnownEvent(Context, Uid, Cursor))
 		{
@@ -4922,22 +4889,26 @@ int32 FProtocol5Stage::DispatchEvents(const FMachineContext& Context, const FEve
 		{
 			UE_TRACE_ANALYSIS_DEBUG_LOG("Warning: Unexpected event with UID %u (0x%X)", Uid, Uid);
 			Context.EmitMessagef(EAnalysisMessageSeverity::Warning, TEXT("An unknown event UID (%u) was encountered."), Uid);
-#if UE_TRACE_ANALYSIS_DEBUG && UE_TRACE_ANALYSIS_DEBUG_LEVEL >= 3
-			UE_TRACE_ANALYSIS_DEBUG_BeginStringBuilder();
-			const uint8* StartCursor = Cursor->Data;
-			for (uint32 i = 0; i < 32 && i < Cursor->EventSize; ++i)
-			{
-				UE_TRACE_ANALYSIS_DEBUG_Appendf("%02X ", StartCursor[i]);
-			}
-			UE_TRACE_ANALYSIS_DEBUG_EndStringBuilder();
-			UE_TRACE_ANALYSIS_DEBUG_ResetStringBuilder();
-			UE_TRACE_ANALYSIS_DEBUG_Append("[[[");
-			for (uint32 i = 0; i < 128 && i < Cursor->EventSize; ++i)
+
+#if UE_TRACE_ANALYSIS_DEBUG
+			if (UE::Trace::Analysis::GDebugLevel >= 3)
 			{
-				UE_TRACE_ANALYSIS_DEBUG_AppendChar((char)StartCursor[i]);
+				UE_TRACE_ANALYSIS_DEBUG_BeginStringBuilder();
+				const uint8* StartCursor = Cursor->Data;
+				for (uint32 i = 0; i < 32 && i < Cursor->EventSize; ++i)
+				{
+					UE_TRACE_ANALYSIS_DEBUG_Appendf("%02X ", StartCursor[i]);
+				}
+				UE_TRACE_ANALYSIS_DEBUG_EndStringBuilder();
+				UE_TRACE_ANALYSIS_DEBUG_ResetStringBuilder();
+				UE_TRACE_ANALYSIS_DEBUG_Append("[[[");
+				for (uint32 i = 0; i < 128 && i < Cursor->EventSize; ++i)
+				{
+					UE_TRACE_ANALYSIS_DEBUG_AppendChar((char)StartCursor[i]);
+				}
+				UE_TRACE_ANALYSIS_DEBUG_Append("]]]");
+				UE_TRACE_ANALYSIS_DEBUG_EndStringBuilder();
 			}
-			UE_TRACE_ANALYSIS_DEBUG_Append("]]]");
-			UE_TRACE_ANALYSIS_DEBUG_EndStringBuilder();
 #endif // UE_TRACE_ANALYSIS_DEBUG
 
 			return -1;
diff --git a/Engine/Source/Developer/TraceAnalysis/Private/Analysis/Transport/TidPacketTransport.cpp b/Engine/Source/Developer/TraceAnalysis/Private/Analysis/Transport/TidPacketTransport.cpp
--- a/Engine/Source/Developer/TraceAnalysis/Private/Analysis/Transport/TidPacketTransport.cpp
+++ b/Engine/Source/Developer/TraceAnalysis/Private/Analysis/Transport/TidPacketTransport.cpp
@@ -80,7 +80,7 @@ FTidPacketTransport::EReadPacketResult FTidPacketTransport::ReadPacket()
 	bool bIsPartial = !!(PacketBase->ThreadId & FTidPacketBase::PartialMarker);
 	FThreadStream* Thread = FindOrAddThread(ThreadId, !bIsPartial);
 #endif
-	
+
 	if (Thread == nullptr)
 	{
 		return EReadPacketResult::Continue;
@@ -93,8 +93,7 @@ FTidPacketTransport::EReadPacketResult FTidPacketTransport::ReadPacket()
 		uint16 DecodedSize = Packet->DecodedSize;
 
 #if UE_TRACE_ANALYSIS_DEBUG
-#if UE_TRACE_ANALYSIS_DEBUG_LEVEL >= 3
-		UE_TRACE_ANALYSIS_DEBUG_LOG("[PACKET %u] Tid=%u, Size: %u + %u bytes, DecodedSize: %u bytes (%.0f%%)%s",
+		UE_TRACE_ANALYSIS_DEBUG_LOG_EX(3, "[PACKET %u] Tid=%u, Size: %u + %u bytes, DecodedSize: %u bytes (%.0f%%)%s",
 			NumPackets,
 			ThreadId,
 			uint32(sizeof(FTidPacketEncoded)),
@@ -102,7 +101,6 @@ FTidPacketTransport::EReadPacketResult FTidPacketTransport::ReadPacket()
 			DecodedSize,
 			((double)DataSize * 100.0) / (double)DecodedSize - 100.0,
 			((DataSize == 0) || (DataSize > DecodedSize)) ? " !!!" : "");
-#endif // UE_TRACE_ANALYSIS_DEBUG_LEVEL
 		TotalPacketHeaderSize += sizeof(FTidPacketEncoded) - sizeof(FTidPacketBase);
 		TotalDecodedSize += DecodedSize;
 		uint64& TotalDataSizePerThread = DataSizePerThread.FindOrAdd(ThreadId);
@@ -124,17 +122,20 @@ FTidPacketTransport::EReadPacketResult FTidPacketTransport::ReadPacket()
 				return EReadPacketResult::ReadError;
 			}
 
-#if UE_TRACE_ANALYSIS_DEBUG && UE_TRACE_ANALYSIS_DEBUG_LEVEL >= 4
-			UE_TRACE_ANALYSIS_DEBUG_BeginStringBuilder();
-			for (uint32 i = 0; i < 32 && i < DecodedSize; ++i)
-			{
-				UE_TRACE_ANALYSIS_DEBUG_Appendf("%02X ", Dest[i]);
-			}
-			if (DecodedSize > 32)
+#if UE_TRACE_ANALYSIS_DEBUG
+			if (UE::Trace::Analysis::GDebugLevel >= 4)
 			{
-				UE_TRACE_ANALYSIS_DEBUG_Append("...");
+				UE_TRACE_ANALYSIS_DEBUG_BeginStringBuilder();
+				for (uint32 i = 0; i < 32 && i < DecodedSize; ++i)
+				{
+					UE_TRACE_ANALYSIS_DEBUG_Appendf("%02X ", Dest[i]);
+				}
+				if (DecodedSize > 32)
+				{
+					UE_TRACE_ANALYSIS_DEBUG_Append("...");
+				}
+				UE_TRACE_ANALYSIS_DEBUG_EndStringBuilder();
 			}
-			UE_TRACE_ANALYSIS_DEBUG_EndStringBuilder();
 #endif // UE_TRACE_ANALYSIS_DEBUG
 		}
 		else // DataSize == 0
@@ -145,13 +146,11 @@ FTidPacketTransport::EReadPacketResult FTidPacketTransport::ReadPacket()
 	else // not encoded
 	{
 #if UE_TRACE_ANALYSIS_DEBUG
-#if UE_TRACE_ANALYSIS_DEBUG_LEVEL >= 3
-		UE_TRACE_ANALYSIS_DEBUG_LOG("[PACKET %u] Tid=%u, Size: %u + %u bytes",
+		UE_TRACE_ANALYSIS_DEBUG_LOG_EX(3, "[PACKET %u] Tid=%u, Size: %u + %u bytes",
 			NumPackets,
 			ThreadId,
 			uint32(sizeof(FTidPacket)),
 			DataSize);
-#endif // UE_TRACE_ANALYSIS_DEBUG_LEVEL
 		TotalPacketHeaderSize += sizeof(FTidPacket) - sizeof(FTidPacketBase);
 		TotalDecodedSize += DataSize;
 		uint64& TotalDataSizePerThread = DataSizePerThread.FindOrAdd(ThreadId);
@@ -163,17 +162,20 @@ FTidPacketTransport::EReadPacketResult FTidPacketTransport::ReadPacket()
 
 		Thread->Buffer.Append((uint8*)(PacketBase + 1), DataSize);
 
-#if UE_TRACE_ANALYSIS_DEBUG && UE_TRACE_ANALYSIS_DEBUG_LEVEL >= 4
-		UE_TRACE_ANALYSIS_DEBUG_BeginStringBuilder();
-		for (uint32 i = 0; i < 32 && i < DataSize; ++i)
-		{
-			UE_TRACE_ANALYSIS_DEBUG_Appendf("%02X ", ((uint8*)(PacketBase + 1))[i]);
-		}
-		if (DataSize > 32)
+#if UE_TRACE_ANALYSIS_DEBUG
+		if (UE::Trace::Analysis::GDebugLevel >= 4)
 		{
-			UE_TRACE_ANALYSIS_DEBUG_Append("...");
+			UE_TRACE_ANALYSIS_DEBUG_BeginStringBuilder();
+			for (uint32 i = 0; i < 32 && i < DataSize; ++i)
+			{
+				UE_TRACE_ANALYSIS_DEBUG_Appendf("%02X ", ((uint8*)(PacketBase + 1))[i]);
+			}
+			if (DataSize > 32)
+			{
+				UE_TRACE_ANALYSIS_DEBUG_Append("...");
+			}
+			UE_TRACE_ANALYSIS_DEBUG_EndStringBuilder();
 		}
-		UE_TRACE_ANALYSIS_DEBUG_EndStringBuilder();
 #endif // UE_TRACE_ANALYSIS_DEBUG
 	}
 
@@ -183,14 +185,14 @@ FTidPacketTransport::EReadPacketResult FTidPacketTransport::ReadPacket()
 		uint64 PacketSerial = *GetPointer<uint64>();
 		if ((LastPacketSerial + 1) != PacketSerial)
 		{
-			UE_LOG(LogCore, Error, TEXT("Found packet with index '%llu' when '%llu` was expected."), PacketSerial, LastPacketSerial);
+			UE_LOG(LogCore, Error, TEXT("Found packet with index '%llu' when '%llu' was expected."), PacketSerial, LastPacketSerial);
 			return EReadPacketResult::ReadError;
 		}
 		LastPacketSerial = PacketSerial;
 		FTransport::Advance(sizeof(uint64));
 	}
 #endif
-	
+
 	return EReadPacketResult::Continue;
 }
 
@@ -207,7 +209,7 @@ FTidPacketTransport::FThreadStream* FTidPacketTransport::FindOrAddThread(
 			return &(Threads[i]);
 		}
 	}
-	
+
 	if (!bAddIfNotFound)
 	{
 		return nullptr;
diff --git a/Engine/Source/Developer/TraceAnalysis/Private/TraceAnalysisDebug.cpp b/Engine/Source/Developer/TraceAnalysis/Private/TraceAnalysisDebug.cpp
--- a/Engine/Source/Developer/TraceAnalysis/Private/TraceAnalysisDebug.cpp
+++ b/Engine/Source/Developer/TraceAnalysis/Private/TraceAnalysisDebug.cpp
@@ -1,3 +1,8 @@
 // Copyright Epic Games, Inc. All Rights Reserved.
 
 #include "TraceAnalysisDebug.h"
+
+namespace UE::Trace::Analysis
+{
+	uint32 GDebugLevel = 0;
+}
diff --git a/Engine/Source/Developer/TraceAnalysis/Public/TraceAnalysisDebug.h b/Engine/Source/Developer/TraceAnalysis/Public/TraceAnalysisDebug.h
--- a/Engine/Source/Developer/TraceAnalysis/Public/TraceAnalysisDebug.h
+++ b/Engine/Source/Developer/TraceAnalysis/Public/TraceAnalysisDebug.h
@@ -2,6 +2,8 @@
 
 #pragma once
 
+#include "CoreTypes.h"
+
 // Enables debug API in UE::TRACE::IAnalyzer class (and sub-types).
 // Enabling it should not add any impact on analysis' performance or behaviour.
 #if !defined(UE_TRACE_ANALYSIS_DEBUG_API)
@@ -87,10 +89,28 @@
 	#define UE_TRACE_ANALYSIS_DEBUG_EndStringBuilder()
 	#define UE_TRACE_ANALYSIS_DEBUG_ResetStringBuilder()
 
-#endif  // UE_TRACE_ANALYSIS_DEBUG
+#endif // UE_TRACE_ANALYSIS_DEBUG
 #endif // !defined(UE_TRACE_ANALYSIS_DEBUG_LOG)
 
-// Log level (1 = minimum, 2 = normal, 3 = verbose, 4 = verbose+). It is used only when UE_TRACE_ANALYSIS_DEBUG is enabled.
-#if !defined(UE_TRACE_ANALYSIS_DEBUG_LEVEL)
-	#define UE_TRACE_ANALYSIS_DEBUG_LEVEL 2
-#endif
+#if UE_TRACE_ANALYSIS_DEBUG
+
+namespace UE::Trace::Analysis
+{
+	// Log level (1 = minimum, 2 = normal, 3 = verbose, 4 = verbose+).
+	// It is used only when UE_TRACE_ANALYSIS_DEBUG is enabled.
+	extern TRACEANALYSIS_API uint32 GDebugLevel;
+}
+
+#define UE_TRACE_ANALYSIS_DEBUG_LOG_EX(level, format, ...) \
+	{ \
+		if (UE::Trace::Analysis::GDebugLevel >= level) \
+		{ \
+			UE_TRACE_ANALYSIS_DEBUG_LOG(format, __VA_ARGS__); \
+		} \
+	}
+
+#else // UE_TRACE_ANALYSIS_DEBUG
+
+#define UE_TRACE_ANALYSIS_DEBUG_LOG_EX(level, format, ...)
+
+#endif // UE_TRACE_ANALYSIS_DEBUG
diff --git a/Engine/Source/Programs/TraceAnalyzer/Private/ConvertToText.cpp b/Engine/Source/Programs/TraceAnalyzer/Private/ConvertToText.cpp
--- a/Engine/Source/Programs/TraceAnalyzer/Private/ConvertToText.cpp
+++ b/Engine/Source/Programs/TraceAnalyzer/Private/ConvertToText.cpp
@@ -12,6 +12,7 @@
 #include "Trace/Analysis.h"
 #include "Trace/Analyzer.h"
 #include "Trace/DataStream.h"
+#include "TraceAnalysisDebug.h"
 
 // TraceAnalyzer
 #include "Command.h"
@@ -1976,6 +1977,7 @@ static void PrintConvertToTextUsage()
 	puts("    -no_event_log      : no 'event' logging");
 	puts("    -no_analysis_stats : no analysis stats");
 	puts("    -no_event_stats    : no event stats");
+	puts("    -analysis_log_level=<level> : 1=minimum, 2=normal (default), 3=verbose, 4=verbose+");
 }
 
 ////////////////////////////////////////////////////////////////////////////////////////////////////
@@ -1997,6 +1999,7 @@ static int32 ConvertToTextMain(int32 ArgC, TCHAR const* const* ArgV)
 	bool bNoEventLog = false;
 	bool bNoAnalysisStats = false;
 	bool bNoEventStats = false;
+	uint32 AnalysisLogLevel = 2;
 
 	Input = OpenFile(TraceFileName, false);
 	if (Input < 0)
@@ -2034,12 +2037,20 @@ static int32 ConvertToTextMain(int32 ArgC, TCHAR const* const* ArgV)
 		{
 			bNoEventStats = true;
 		}
+		else if (FCString::Strnicmp(TEXT("-analysis_log_level="), ArgV[ArgIndex], 20) == 0)
+		{
+			AnalysisLogLevel = FCString::Atoi(ArgV[ArgIndex] + 20);
+		}
 		else
 		{
 			LogTraceAnalyzerWarning("Unknown cmd line argument '%s'!", ArgV[ArgIndex]);
 		}
 	}
 
+#if UE_TRACE_ANALYSIS_DEBUG
+	UE::Trace::Analysis::GDebugLevel = AnalysisLogLevel;
+#endif
+
 	struct FDataStream : public UE::Trace::IInDataStream
 	{
 		virtual int32 Read(void* Data, uint32 Size) override
diff --git a/Engine/Source/Programs/TraceAnalyzer/Private/TextSerializer.cpp b/Engine/Source/Programs/TraceAnalyzer/Private/TextSerializer.cpp
--- a/Engine/Source/Programs/TraceAnalyzer/Private/TextSerializer.cpp
+++ b/Engine/Source/Programs/TraceAnalyzer/Private/TextSerializer.cpp
@@ -6,6 +6,9 @@
 #include "Math/NumericLimits.h"
 #include "Misc/CString.h"
 
+// TraceAnalysis
+#include "TraceAnalysisDebug.h"
+
 #include <limits>
 #include <cmath>
 
@@ -198,8 +201,11 @@ void FTextSerializer::WriteValueBinary(const void* Data, uint32 Size)
 
 FStdoutTextSerializer::FStdoutTextSerializer()
 {
-#if UE_TRACE_ANALYSIS_DEBUG && (UE_TRACE_ANALYSIS_DEBUG_LOG_IMPL <= 2) && (UE_TRACE_ANALYSIS_DEBUG_LEVEL >= 2)
-	bWriteEventHeader = false;
+#if UE_TRACE_ANALYSIS_DEBUG && (UE_TRACE_ANALYSIS_DEBUG_LOG_IMPL <= 2)
+	if (UE::Trace::Analysis::GDebugLevel >= 2)
+	{
+		bWriteEventHeader = false;
+	}
 #endif
 }
 
diff --git a/Engine/Source/Programs/TraceAnalyzer/TraceAnalyzer.Target.cs b/Engine/Source/Programs/TraceAnalyzer/TraceAnalyzer.Target.cs
--- a/Engine/Source/Programs/TraceAnalyzer/TraceAnalyzer.Target.cs
+++ b/Engine/Source/Programs/TraceAnalyzer/TraceAnalyzer.Target.cs
@@ -48,7 +48,6 @@ public class TraceAnalyzerTarget : TargetRules
 		GlobalDefinitions.Add("UE_TRACE_ANALYSIS_DEBUG_API=1");
 		GlobalDefinitions.Add("UE_TRACE_ANALYSIS_DEBUG=1");
 		GlobalDefinitions.Add("UE_TRACE_ANALYSIS_DEBUG_LOG_IMPL=1");
-		GlobalDefinitions.Add("UE_TRACE_ANALYSIS_DEBUG_LEVEL=2");
 
 		//GlobalDefinitions.Add("_CRT_SECURE_NO_WARNINGS");
 		//GlobalDefinitions.Add("_CRT_NONSTDC_NO_WARNINGS");
