[Insights] Support for analysing OnConnect events.

Adds analysis support for OnConnect thread events. These events are processed and dispatched at the same time as Important events. Sync events emitted here are expected to have serial value `0` and a treated as it the event was non-sync.

#Synced-CL 52616121

diff --git a/Engine/Source/Developer/TraceAnalysis/Private/Analysis/Engine.cpp b/Engine/Source/Developer/TraceAnalysis/Private/Analysis/Engine.cpp
--- a/Engine/Source/Developer/TraceAnalysis/Private/Analysis/Engine.cpp
+++ b/Engine/Source/Developer/TraceAnalysis/Private/Analysis/Engine.cpp
@@ -3581,6 +3581,7 @@ protected:
 	virtual EStatus			OnData(FStreamReader& Reader, const FMachineContext& Context) override;
 	EStatus					OnDataNewEvents(const FMachineContext& Context);
 	EStatus					OnDataImportant(const FMachineContext& Context);
+	EStatus					OnDataPseudoImportant(const FMachineContext& Context);
 	EStatus					OnDataNormal(const FMachineContext& Context);
 	int32					ParseImportantEvents(FStreamReader& Reader, EventDescArray& OutEventDescs, const FMachineContext& Context);
 	int32					ParseEvents(FStreamReader& Reader, EventDescArray& OutEventDescs, const FMachineContext& Context);
@@ -3676,6 +3677,15 @@ FProtocol5Stage::EStatus FProtocol5Stage::OnData(
 		}
 		bool bNotEnoughData = (Ret == EStatus::NotEnoughData);
 
+		// Pseudo Important events, events that should be treated
+		// as important but written like regular event (e.g. during OnConnect).
+		Ret = OnDataPseudoImportant(Context);
+		if (Ret == EStatus::Error)
+		{
+			return Ret;
+		}
+		bNotEnoughData |= (Ret == EStatus::NotEnoughData);
+
 		// Normal events
 		Ret = OnDataNormal(Context);
 		if (Ret == EStatus::Error)
@@ -3704,7 +3714,7 @@ FProtocol5Stage::EStatus FProtocol5Stage::OnDataNewEvents(const FMachineContext&
 {
 	EventDescs.Reset();
 
-	FStreamReader* ThreadReader = Transport.GetThreadStream(ETransportTid::Events);
+	FStreamReader* ThreadReader = Transport.GetKnownStream(FTidPacketTransport::EThreadStreamIndex::Events);
 	if (ThreadReader->IsEmpty())
 	{
 		return EStatus::EndOfStream;
@@ -3735,7 +3745,7 @@ FProtocol5Stage::EStatus FProtocol5Stage::OnDataImportant(const FMachineContext&
 
 	EventDescs.Reset();
 
-	FStreamReader* ThreadReader = Transport.GetThreadStream(ETransportTid::Importants);
+	FStreamReader* ThreadReader = Transport.GetKnownStream(FTidPacketTransport::EThreadStreamIndex::Importants);
 	if (ThreadReader->IsEmpty())
 	{
 		return EStatus::EndOfStream;
@@ -3769,6 +3779,61 @@ FProtocol5Stage::EStatus FProtocol5Stage::OnDataImportant(const FMachineContext&
 	return bNotEnoughData ? EStatus::NotEnoughData : EStatus::EndOfStream;
 }
 
+////////////////////////////////////////////////////////////////////////////////
+FProtocol5Stage::EStatus FProtocol5Stage::OnDataPseudoImportant(const FMachineContext& Context)
+{
+	EventDescs.Reset();
+
+	FStreamReader* ThreadReader = Transport.GetKnownStream(FTidPacketTransport::EThreadStreamIndex::PseudoImportants);
+	if (ThreadReader->IsEmpty())
+	{
+		return EStatus::EndOfStream;
+	}
+
+	if (ParseEvents(*ThreadReader, EventDescs, Context) < 0)
+	{
+		UE_TRACE_ANALYSIS_DEBUG_LOG("Error: Failed to parse on connect events");
+		return EStatus::Error;
+	}
+
+	bool bNotEnoughData = !ThreadReader->IsEmpty();
+
+	if (EventDescs.Num() <= 0)
+	{
+		return bNotEnoughData ? EStatus::NotEnoughData : EStatus::EndOfStream;
+	}
+	
+	// Make sure serial events on the OnConnect thread is processed as 
+	// no-sync events.
+	for (auto& Event : EventDescs)
+	{
+		if (Event.Serial != ESerial::Ignored)
+		{
+			// TraceLog emits serial events on the OnConnect thread
+			// with serial value '0'.
+			if (Event.Serial != 0)
+			{
+				UE_TRACE_ANALYSIS_DEBUG_LOG("Warning: An event in PseudoImportant thread contained a non-zero serial.");
+			}
+			Event.Serial = ESerial::Ignored;
+		}
+	}
+
+	// Dispatch looks ahead to the next desc looking for runs of aux blobs. As
+	// such we should add a terminal desc for it to read. Note the "- 1" too.
+	FEventDesc& EventDesc = EventDescs.Emplace_GetRef();
+	EventDesc.Serial = ESerial::Terminal;
+
+	Context.Bridge.SetActiveThread(ETransportTid::PseudoImportants);
+	if (DispatchEvents(Context, EventDescs.GetData(), EventDescs.Num() - 1) < 0)
+	{
+		UE_TRACE_ANALYSIS_DEBUG_LOG("Error: Failed to dispatch pseudo important events");
+		return EStatus::Error;
+	}
+
+	return bNotEnoughData ? EStatus::NotEnoughData : EStatus::EndOfStream;
+}
+
 ////////////////////////////////////////////////////////////////////////////////
 int32 FProtocol5Stage::ParseImportantEvents(FStreamReader& Reader, EventDescArray& OutEventDescs, const FMachineContext& Context)
 {
@@ -3877,7 +3942,7 @@ FProtocol5Stage::EStatus FProtocol5Stage::OnDataNormal(const FMachineContext& Co
 
 	bSkipSerial = false;
 
-	for (uint32 i = ETransportTid::Bias, n = Transport.GetThreadCount(); i < n; ++i)
+	for (uint32 i = (uint32)FTidPacketTransport::EThreadStreamIndex::RealThreadsStart, n = Transport.GetThreadCount(); i < n; ++i)
 	{
 		uint32 NumEventDescs = EventDescs.Num();
 
diff --git a/Engine/Source/Developer/TraceAnalysis/Private/Analysis/Transport/TidPacketTransport.cpp b/Engine/Source/Developer/TraceAnalysis/Private/Analysis/Transport/TidPacketTransport.cpp
--- a/Engine/Source/Developer/TraceAnalysis/Private/Analysis/Transport/TidPacketTransport.cpp
+++ b/Engine/Source/Developer/TraceAnalysis/Private/Analysis/Transport/TidPacketTransport.cpp
@@ -232,7 +232,7 @@ FTidPacketTransport::ETransportResult FTidPacketTransport::Update()
 
 	Threads.RemoveAll([] (const FThreadStream& Thread)
 	{
-		return (Thread.ThreadId <= ETransportTid::Importants) ? false : Thread.Buffer.IsEmpty();
+		return (Thread.ThreadId < (uint32)FTidPacketTransport::EThreadStreamIndex::RealThreadsStart) ? false : Thread.Buffer.IsEmpty();
 	});
 
 #if UE_TRACE_ANALYSIS_DEBUG
@@ -362,6 +362,13 @@ FStreamReader* FTidPacketTransport::GetThreadStream(uint32 Index)
 	return &(Threads[Index].Buffer);
 }
 
+////////////////////////////////////////////////////////////////////////////////
+FStreamReader* FTidPacketTransport::GetKnownStream(EThreadStreamIndex KnownThread)
+{
+	check(KnownThread < EThreadStreamIndex::RealThreadsStart);
+	return GetThreadStream((uint32)KnownThread);
+}
+
 ////////////////////////////////////////////////////////////////////////////////
 uint32 FTidPacketTransport::GetThreadId(uint32 Index) const
 {
diff --git a/Engine/Source/Developer/TraceAnalysis/Private/Analysis/Transport/TidPacketTransport.h b/Engine/Source/Developer/TraceAnalysis/Private/Analysis/Transport/TidPacketTransport.h
--- a/Engine/Source/Developer/TraceAnalysis/Private/Analysis/Transport/TidPacketTransport.h
+++ b/Engine/Source/Developer/TraceAnalysis/Private/Analysis/Transport/TidPacketTransport.h
@@ -25,6 +25,19 @@ public:
 		Error
 	};
 
+	/**
+	 * The "pseudo threads" have a fixed
+	 * index in the thread stream array.
+	 */
+	enum class EThreadStreamIndex : uint32
+	{
+		Events = 0,
+		Importants = 1,
+		PseudoImportants = 2,
+		RealThreadsStart = 3,
+	};
+
+
 							FTidPacketTransport() {}
 	virtual					~FTidPacketTransport() = default;
 
@@ -35,9 +48,11 @@ public:
 	ETransportResult		Update();
 	uint32					GetThreadCount() const;
 	FStreamReader*			GetThreadStream(uint32 Index);
+	FStreamReader*			GetKnownStream(EThreadStreamIndex KnownThread);
 	uint32					GetThreadId(uint32 Index) const;
 	uint32					GetSyncCount() const;
 
+
 protected:
 	uint32					Synced = 0x7fff'ffff;
 	uint32					NumPackets = 0;
@@ -66,6 +81,7 @@ private:
 	TArray<FThreadStream>	Threads = {
 								{ {}, ETransportTid::Events },
 								{ {}, ETransportTid::Importants },
+								{ {}, ETransportTid::PseudoImportants }
 							};
 
 #if UE_TRACE_ANALYSIS_DEBUG
diff --git a/Engine/Source/Developer/TraceAnalysis/Private/TraceWriter.cpp b/Engine/Source/Developer/TraceAnalysis/Private/TraceWriter.cpp
--- a/Engine/Source/Developer/TraceAnalysis/Private/TraceWriter.cpp
+++ b/Engine/Source/Developer/TraceAnalysis/Private/TraceWriter.cpp
@@ -2,6 +2,7 @@
 
 #include "Trace/TraceWriter.h"
 
+#include "Analysis/Transport/TidPacketTransport.h"
 #include "HAL/PlatformTime.h"
 
 // TraceLog
@@ -495,9 +496,10 @@ void FTraceWriter::Begin(ETraceWriterBeginOptions Options)
 	Threads.Reset();
 	Threads.Add({ 0, FAnsiString("NewEvents") });
 	Threads.Add({ 1, FAnsiString("Importants") });
+	Threads.Add({ 2, FAnsiString("OnConnect") });
 	static_assert(ETransportTid::Events == 0);
 	static_assert(ETransportTid::Importants == 1);
-	check(Threads.Num() == ETransportTid::Bias);
+	check(Threads.Num() == (uint32)FTidPacketTransport::EThreadStreamIndex::RealThreadsStart);
 
 	constexpr uint16 GControlPort = 0;
 
@@ -1106,7 +1108,7 @@ void FTraceWriter::WriteEvent(const FTraceWriterEventBuilder& Builder)
 		if (CurrentThreadId == uint32(ETransportTid::Importants))
 		{
 			SetError(TEXT("Non-important events can not be emitted on the Importants thread!"));
-			SetCurrentThread(uint32(ETransportTid::Bias));
+			SetCurrentThread(uint32(FTidPacketTransport::EThreadStreamIndex::RealThreadsStart));
 		}
 	}
 
