[Insights] Fix issue introduced with 52678429

#rnx
#Synced-CL 52678544

diff --git a/Engine/Source/Developer/TraceAnalysis/Private/Analysis/Transport/TidPacketTransport.cpp b/Engine/Source/Developer/TraceAnalysis/Private/Analysis/Transport/TidPacketTransport.cpp
--- a/Engine/Source/Developer/TraceAnalysis/Private/Analysis/Transport/TidPacketTransport.cpp
+++ b/Engine/Source/Developer/TraceAnalysis/Private/Analysis/Transport/TidPacketTransport.cpp
@@ -232,7 +232,10 @@ FTidPacketTransport::ETransportResult FTidPacketTransport::Update()
 
 	Threads.RemoveAll([] (const FThreadStream& Thread)
 	{
-		return (Thread.ThreadId < (uint32)FTidPacketTransport::EThreadStreamIndex::RealThreadsStart) ? false : Thread.Buffer.IsEmpty();
+		return Thread.ThreadId != ETransportTid::Events &&
+				Thread.ThreadId != ETransportTid::Importants &&
+				Thread.ThreadId != ETransportTid::PseudoImportants &&
+				Thread.Buffer.IsEmpty();
 	});
 
 #if UE_TRACE_ANALYSIS_DEBUG
